Fixed build break.

diff --git a/Engine/Plugins/Runtime/TextureShare/Source/TextureShare/Private/TextureShareModule.cpp b/Engine/Plugins/Runtime/TextureShare/Source/TextureShare/Private/TextureShareModule.cpp
--- a/Engine/Plugins/Runtime/TextureShare/Source/TextureShare/Private/TextureShareModule.cpp
+++ b/Engine/Plugins/Runtime/TextureShare/Source/TextureShare/Private/TextureShareModule.cpp
@@ -257,8 +257,6 @@ bool FTextureShareModule::SendSceneContext_RenderThread(FRHICommandListImmediate
 
 			SendTexture_RenderThread(RHICmdList, ShareItem, TextureShareStrings::texture_name::GBufferVelocity, SceneContext.SceneVelocity);
 			SendTexture_RenderThread(RHICmdList, ShareItem, TextureShareStrings::texture_name::Foveation, SceneContext.FoveationTexture);
-
-			SendTexture_RenderThread(RHICmdList, ShareItem, TextureShareStrings::texture_name::DirectionalOcclusion, SceneContext.DirectionalOcclusion);
 			//@todo: Add more textures
 
 			return true;
diff --git a/Engine/Plugins/Runtime/TextureShare/Source/TextureShare/Private/TextureShareStrings.h b/Engine/Plugins/Runtime/TextureShare/Source/TextureShare/Private/TextureShareStrings.h
--- a/Engine/Plugins/Runtime/TextureShare/Source/TextureShare/Private/TextureShareStrings.h
+++ b/Engine/Plugins/Runtime/TextureShare/Source/TextureShare/Private/TextureShareStrings.h
@@ -28,8 +28,6 @@ namespace TextureShareStrings
 		static constexpr auto GBufferVelocity = TEXT("GBufferVelocity");
 		static constexpr auto Foveation = TEXT("Foveation");
 
-		static constexpr auto DirectionalOcclusion = TEXT("DirectionalOcclusion");
-
 		static constexpr auto BackBuffer = TEXT("BackBuffer");
 	}
 };
