Fixing PVS 7.7 Issues:

warning V530: The return value of function 'InitializeSecurityDescriptor' is required to be utilized.
warning V530: The return value of function 'SetSecurityDescriptorDacl' is required to be utilized.

diff --git a/Engine/Plugins/Runtime/TextureShare/Source/TextureShareD3D12/Private/Platform/TextureShareD3D12PlatformWindows.cpp b/Engine/Plugins/Runtime/TextureShare/Source/TextureShareD3D12/Private/Platform/TextureShareD3D12PlatformWindows.cpp
--- a/Engine/Plugins/Runtime/TextureShare/Source/TextureShareD3D12/Private/Platform/TextureShareD3D12PlatformWindows.cpp
+++ b/Engine/Plugins/Runtime/TextureShare/Source/TextureShareD3D12/Private/Platform/TextureShareD3D12PlatformWindows.cpp
@@ -14,7 +14,8 @@ FD3D12CrossGPUHeapSecurityAttributes::FD3D12CrossGPUHeapSecurityAttributes()
 	PSID* ppSID = (PSID*)((PBYTE)m_winPSecurityDescriptor + SECURITY_DESCRIPTOR_MIN_LENGTH);
 	PACL* ppACL = (PACL*)((PBYTE)ppSID + sizeof(PSID*));
 
-	InitializeSecurityDescriptor(m_winPSecurityDescriptor, SECURITY_DESCRIPTOR_REVISION);
+	BOOL retval = InitializeSecurityDescriptor(m_winPSecurityDescriptor, SECURITY_DESCRIPTOR_REVISION);
+	(retval);
 
 	SID_IDENTIFIER_AUTHORITY sidIdentifierAuthority = SECURITY_WORLD_SID_AUTHORITY;
 	AllocateAndInitializeSid(&sidIdentifierAuthority, 1, SECURITY_WORLD_RID, 0, 0, 0, 0, 0, 0, 0, ppSID);
@@ -30,7 +31,8 @@ FD3D12CrossGPUHeapSecurityAttributes::FD3D12CrossGPUHeapSecurityAttributes()
 
 	SetEntriesInAcl(1, &explicitAccess, NULL, ppACL);
 
-	SetSecurityDescriptorDacl(m_winPSecurityDescriptor, true, *ppACL, false);
+	retval = SetSecurityDescriptorDacl(m_winPSecurityDescriptor, true, *ppACL, false);
+	(retval);
 
 	m_winSecurityAttributes.nLength = sizeof(m_winSecurityAttributes);
 	m_winSecurityAttributes.lpSecurityDescriptor = m_winPSecurityDescriptor;
