Remove reference to HoloLens - it isn't listed in the .uplugin and isn't supported.

#rnx

diff --git a/Engine/Plugins/Runtime/TextureShare/Source/TextureShareCore/TextureShareCore.Build.cs b/Engine/Plugins/Runtime/TextureShare/Source/TextureShareCore/TextureShareCore.Build.cs
--- a/Engine/Plugins/Runtime/TextureShare/Source/TextureShareCore/TextureShareCore.Build.cs
+++ b/Engine/Plugins/Runtime/TextureShare/Source/TextureShareCore/TextureShareCore.Build.cs
@@ -42,7 +42,7 @@
 
 		string TextureShareLibPlatformName = Target.Platform.ToString();
 
-		if (Target.Platform == UnrealTargetPlatform.Win64 || Target.Platform == UnrealTargetPlatform.HoloLens)
+		if (Target.Platform == UnrealTargetPlatform.Win64)
 		{
 			PublicDefinitions.Add("TEXTURESHARELIB_PLATFORM_WINDOWS=1");
 			TextureShareLibPlatformName = "Windows";
diff --git a/Engine/Plugins/Runtime/TextureShare/Source/TextureShareD3D11/TextureShareD3D11.Build.cs b/Engine/Plugins/Runtime/TextureShare/Source/TextureShareD3D11/TextureShareD3D11.Build.cs
--- a/Engine/Plugins/Runtime/TextureShare/Source/TextureShareD3D11/TextureShareD3D11.Build.cs
+++ b/Engine/Plugins/Runtime/TextureShare/Source/TextureShareD3D11/TextureShareD3D11.Build.cs
@@ -1,31 +1,30 @@
-// Copyright Epic Games, Inc. All Rights Reserved.
-
-using UnrealBuildTool;
-
-public class TextureShareD3D11 : ModuleRules
-{
-	public TextureShareD3D11(ReadOnlyTargetRules Target) : base(Target)
-	{
-		PrivateDependencyModuleNames.AddRange(new string[] {
-			"Core",
-			"Engine",
-			"RHI",
-			"RenderCore",
-			"D3D11RHI",
-		});
-
-		///////////////////////////////////////////////////////////////
-		// Platform specific defines
-		///////////////////////////////////////////////////////////////
-		if (!Target.Platform.IsInGroup(UnrealPlatformGroup.Windows))
-		{
-			PrecompileForTargets = PrecompileTargetsType.None;
-		}
-
-		if (Target.Platform.IsInGroup(UnrealPlatformGroup.Windows) ||
-			Target.Platform == UnrealTargetPlatform.HoloLens)
-		{
-			AddEngineThirdPartyPrivateStaticDependencies(Target, "DX11");
-		}
-	}
-}
+// Copyright Epic Games, Inc. All Rights Reserved.
+
+using UnrealBuildTool;
+
+public class TextureShareD3D11 : ModuleRules
+{
+	public TextureShareD3D11(ReadOnlyTargetRules Target) : base(Target)
+	{
+		PrivateDependencyModuleNames.AddRange(new string[] {
+			"Core",
+			"Engine",
+			"RHI",
+			"RenderCore",
+			"D3D11RHI",
+		});
+
+		///////////////////////////////////////////////////////////////
+		// Platform specific defines
+		///////////////////////////////////////////////////////////////
+		if (!Target.Platform.IsInGroup(UnrealPlatformGroup.Windows))
+		{
+			PrecompileForTargets = PrecompileTargetsType.None;
+		}
+
+		if (Target.Platform.IsInGroup(UnrealPlatformGroup.Windows))
+		{
+			AddEngineThirdPartyPrivateStaticDependencies(Target, "DX11");
+		}
+	}
+}
diff --git a/Engine/Plugins/Runtime/TextureShare/Source/TextureShareD3D12/TextureShareD3D12.Build.cs b/Engine/Plugins/Runtime/TextureShare/Source/TextureShareD3D12/TextureShareD3D12.Build.cs
--- a/Engine/Plugins/Runtime/TextureShare/Source/TextureShareD3D12/TextureShareD3D12.Build.cs
+++ b/Engine/Plugins/Runtime/TextureShare/Source/TextureShareD3D12/TextureShareD3D12.Build.cs
@@ -1,35 +1,35 @@
-// Copyright Epic Games, Inc. All Rights Reserved.
-
-using UnrealBuildTool;
-
-public class TextureShareD3D12 : ModuleRules
-{
-	public TextureShareD3D12(ReadOnlyTargetRules Target) : base(Target)
-	{
-		PrivateDependencyModuleNames.AddRange(
-						new string[]
-						{
-							"Core",
-							"Engine",
-							"RenderCore",
-							"RHI",
-							"D3D12RHI",
-						});
-
-		// Allow D3D12 Cross GPU Heap resource API (experimental)
-		PublicDefinitions.Add("TEXTURESHARE_CROSSGPUHEAP=0");
-
-		///////////////////////////////////////////////////////////////
-		// Platform specific defines
-		///////////////////////////////////////////////////////////////
-		if (!Target.Platform.IsInGroup(UnrealPlatformGroup.Windows))
-		{
-			PrecompileForTargets = PrecompileTargetsType.None;
-		}
-
-		if (Target.Platform.IsInGroup(UnrealPlatformGroup.Windows) || Target.Platform == UnrealTargetPlatform.HoloLens)
-		{
-			AddEngineThirdPartyPrivateStaticDependencies(Target, "DX12");
-		}
-	}
-}
+// Copyright Epic Games, Inc. All Rights Reserved.
+
+using UnrealBuildTool;
+
+public class TextureShareD3D12 : ModuleRules
+{
+	public TextureShareD3D12(ReadOnlyTargetRules Target) : base(Target)
+	{
+		PrivateDependencyModuleNames.AddRange(
+						new string[]
+						{
+							"Core",
+							"Engine",
+							"RenderCore",
+							"RHI",
+							"D3D12RHI",
+						});
+
+		// Allow D3D12 Cross GPU Heap resource API (experimental)
+		PublicDefinitions.Add("TEXTURESHARE_CROSSGPUHEAP=0");
+
+		///////////////////////////////////////////////////////////////
+		// Platform specific defines
+		///////////////////////////////////////////////////////////////
+		if (!Target.Platform.IsInGroup(UnrealPlatformGroup.Windows))
+		{
+			PrecompileForTargets = PrecompileTargetsType.None;
+		}
+
+		if (Target.Platform.IsInGroup(UnrealPlatformGroup.Windows))
+		{
+			AddEngineThirdPartyPrivateStaticDependencies(Target, "DX12");
+		}
+	}
+}
