TextureShare:
 * removed dependency on nDisplay plugin
 * Fixed horde Issue 185645 Renamed duplicate class names and file names (compared to nDisplay)
    - 'FDisplayClusterProjectionPolicyBase' to 'FTextureShareProjectionPolicyBase', etc
    - ''DisplayClusterPostprocessStrings.h' to ''TextureSharePostprocessStrings.h' etc

diff --git a/Engine/Plugins/VirtualProduction/TextureShare/Source/TextureShareDisplayCluster/Private/Module/TextureShareDisplayClusterAPI.cpp b/Engine/Plugins/VirtualProduction/TextureShare/Source/TextureShareDisplayCluster/Private/Module/TextureShareDisplayClusterAPI.cpp
--- a/Engine/Plugins/VirtualProduction/TextureShare/Source/TextureShareDisplayCluster/Private/Module/TextureShareDisplayClusterAPI.cpp
+++ b/Engine/Plugins/VirtualProduction/TextureShare/Source/TextureShareDisplayCluster/Private/Module/TextureShareDisplayClusterAPI.cpp
@@ -3,14 +3,14 @@
 #include "Module/TextureShareDisplayClusterAPI.h"
 #include "Module/TextureShareDisplayClusterLog.h"
 
-#include "Projection/TextureShare/DisplayClusterProjectionTextureSharePolicyFactory.h"
-#include "Projection/TextureShare/DisplayClusterProjectionTextureSharePolicy.h"
-#include "Projection/DisplayClusterProjectionStrings.h"
+#include "Projection/TextureShareProjectionPolicyFactory.h"
+#include "Projection/TextureShareProjectionPolicy.h"
+#include "Projection/TextureShareProjectionStrings.h"
 
 #include "Render/Projection/IDisplayClusterProjectionPolicy.h"
 
-#include "PostProcess/TextureShare/DisplayClusterPostprocessTextureShareFactory.h"
-#include "PostProcess/DisplayClusterPostprocessStrings.h"
+#include "PostProcess/TextureSharePostprocessFactory.h"
+#include "PostProcess/TextureSharePostprocessStrings.h"
 
 #include "IDisplayCluster.h"
 
@@ -35,13 +35,13 @@ FTextureShareDisplayClusterAPI::FTextureShareDisplayClusterAPI()
 	TSharedPtr<IDisplayClusterProjectionPolicyFactory> Factory;
 
 	// TextureShare projection
-	Factory = MakeShared<FDisplayClusterProjectionTextureSharePolicyFactory>();
-	ProjectionPolicyFactories.Emplace(DisplayClusterProjectionStrings::Projection::TextureShare, Factory);
+	Factory = MakeShared<FTextureShareProjectionPolicyFactory>();
+	ProjectionPolicyFactories.Emplace(TextureShareProjectionStrings::Projection::TextureShare, Factory);
 
 	TSharedPtr<IDisplayClusterPostProcessFactory> Postprocess;
 
-	Postprocess = MakeShared<FDisplayClusterPostProcessTextureShareFactory>();
-	PostprocessAssets.Emplace(DisplayClusterPostprocessStrings::Postprocess::TextureShare, Postprocess);
+	Postprocess = MakeShared<FTextureSharePostprocessFactory>();
+	PostprocessAssets.Emplace(TextureSharePostprocessStrings::Postprocess::TextureShare, Postprocess);
 }
 
 FTextureShareDisplayClusterAPI::~FTextureShareDisplayClusterAPI()
@@ -124,7 +124,7 @@ bool FTextureShareDisplayClusterAPI::TextureSharePolicySetProjectionData(const T
 	{
 		if (InPolicy.IsValid())
 		{
-			if (FDisplayClusterProjectionTextureSharePolicy* TSPolicyObject = static_cast<FDisplayClusterProjectionTextureSharePolicy*>(InPolicy.Get()))
+			if (FTextureShareProjectionPolicy* TSPolicyObject = static_cast<FTextureShareProjectionPolicy*>(InPolicy.Get()))
 			{
 				return TSPolicyObject->SetCustomProjection(InProjectionData);
 			}
diff --git a/Engine/Plugins/VirtualProduction/TextureShare/Source/TextureShareDisplayCluster/Private/PostProcess/DisplayClusterPostprocessBase.cpp b/Engine/Plugins/VirtualProduction/TextureShare/Source/TextureShareDisplayCluster/Private/PostProcess/DisplayClusterPostprocessBase.cpp
deleted file mode 100644
diff --git a/Engine/Plugins/VirtualProduction/TextureShare/Source/TextureShareDisplayCluster/Private/PostProcess/DisplayClusterPostprocessBase.h b/Engine/Plugins/VirtualProduction/TextureShare/Source/TextureShareDisplayCluster/Private/PostProcess/DisplayClusterPostprocessBase.h
deleted file mode 100644
diff --git a/Engine/Plugins/VirtualProduction/TextureShare/Source/TextureShareDisplayCluster/Private/PostProcess/DisplayClusterPostprocessStrings.h b/Engine/Plugins/VirtualProduction/TextureShare/Source/TextureShareDisplayCluster/Private/PostProcess/DisplayClusterPostprocessStrings.h
deleted file mode 100644
diff --git a/Engine/Plugins/VirtualProduction/TextureShare/Source/TextureShareDisplayCluster/Private/PostProcess/TextureShare/DisplayClusterPostprocessTextureShare.cpp b/Engine/Plugins/VirtualProduction/TextureShare/Source/TextureShareDisplayCluster/Private/PostProcess/TextureShare/DisplayClusterPostprocessTextureShare.cpp
deleted file mode 100644
diff --git a/Engine/Plugins/VirtualProduction/TextureShare/Source/TextureShareDisplayCluster/Private/PostProcess/TextureShare/DisplayClusterPostprocessTextureShare.h b/Engine/Plugins/VirtualProduction/TextureShare/Source/TextureShareDisplayCluster/Private/PostProcess/TextureShare/DisplayClusterPostprocessTextureShare.h
deleted file mode 100644
diff --git a/Engine/Plugins/VirtualProduction/TextureShare/Source/TextureShareDisplayCluster/Private/PostProcess/TextureShare/DisplayClusterPostprocessTextureShareFactory.cpp b/Engine/Plugins/VirtualProduction/TextureShare/Source/TextureShareDisplayCluster/Private/PostProcess/TextureShare/DisplayClusterPostprocessTextureShareFactory.cpp
deleted file mode 100644
diff --git a/Engine/Plugins/VirtualProduction/TextureShare/Source/TextureShareDisplayCluster/Private/PostProcess/TextureShare/DisplayClusterPostprocessTextureShareFactory.h b/Engine/Plugins/VirtualProduction/TextureShare/Source/TextureShareDisplayCluster/Private/PostProcess/TextureShare/DisplayClusterPostprocessTextureShareFactory.h
deleted file mode 100644
diff --git a/Engine/Plugins/VirtualProduction/TextureShare/Source/TextureShareDisplayCluster/Private/PostProcess/TextureShare/DisplayClusterPostprocessTextureShare_Frame.cpp b/Engine/Plugins/VirtualProduction/TextureShare/Source/TextureShareDisplayCluster/Private/PostProcess/TextureShare/DisplayClusterPostprocessTextureShare_Frame.cpp
deleted file mode 100644
diff --git a/Engine/Plugins/VirtualProduction/TextureShare/Source/TextureShareDisplayCluster/Private/PostProcess/TextureShare/DisplayClusterPostprocessTextureShare_FrameProxy.cpp b/Engine/Plugins/VirtualProduction/TextureShare/Source/TextureShareDisplayCluster/Private/PostProcess/TextureShare/DisplayClusterPostprocessTextureShare_FrameProxy.cpp
deleted file mode 100644
diff --git a/Engine/Plugins/VirtualProduction/TextureShare/Source/TextureShareDisplayCluster/Private/PostProcess/TextureSharePostprocess.cpp b/Engine/Plugins/VirtualProduction/TextureShare/Source/TextureShareDisplayCluster/Private/PostProcess/TextureSharePostprocess.cpp
new file mode 100644
--- /dev/null
+++ b/Engine/Plugins/VirtualProduction/TextureShare/Source/TextureShareDisplayCluster/Private/PostProcess/TextureSharePostprocess.cpp
@@ -0,0 +1,214 @@
+// Copyright Epic Games, Inc. All Rights Reserved.
+
+#include "PostProcess/TextureSharePostprocess.h"
+#include "PostProcess/TextureSharePostprocessStrings.h"
+
+#include "Misc/TextureShareDisplayClusterStrings.h"
+
+#include "Module/TextureShareDisplayClusterLog.h"
+
+#include "Containers/TextureShareCoreEnums.h"
+
+#include "ITextureShare.h"
+#include "ITextureShareAPI.h"
+#include "ITextureShareObject.h"
+#include "ITextureShareObjectProxy.h"
+#include "ITextureShareDisplayCluster.h"
+
+#include "Render/Viewport/IDisplayClusterViewportManager.h"
+
+//////////////////////////////////////////////////////////////////////////////////////////////
+namespace DisplayClusterPostProcessTextureShareHelpers
+{
+	static ITextureShareAPI& TextureShareAPI()
+	{
+		static ITextureShareAPI& TextureShareAPISingleton = ITextureShare::Get().GetTextureShareAPI();
+		return TextureShareAPISingleton;
+	}
+
+	static FViewport* GetDisplayViewport(IDisplayClusterViewportManager* InViewportManager)
+	{
+		if (InViewportManager)
+		{
+			if (UWorld* CurrentWorld = InViewportManager->GetCurrentWorld())
+			{
+				if (UGameViewportClient* GameViewportClientPtr = CurrentWorld->GetGameViewport())
+				{
+					return GameViewportClientPtr->Viewport;
+				}
+			}
+		}
+
+		return nullptr;
+	};
+};
+using namespace DisplayClusterPostProcessTextureShareHelpers;
+
+//////////////////////////////////////////////////////////////////////////////////////////////
+// FTextureSharePostprocess
+//////////////////////////////////////////////////////////////////////////////////////////////
+FTextureSharePostprocess::FTextureSharePostprocess(const FString& PostprocessId, const struct FDisplayClusterConfigurationPostprocess* InConfigurationPostprocess)
+	: FTextureSharePostprocessBase(PostprocessId, InConfigurationPostprocess)
+{
+	if (TextureShareAPI().IsObjectExist(TextureShareDisplayClusterStrings::Default::ShareName))
+	{
+		UE_LOG(LogTextureShareDisplayClusterPostProcess, Error, TEXT("TextureShareDisplayCluster: Failed - Already exist"));
+	}
+	else
+	{
+		Object = TextureShareAPI().GetOrCreateObject(TextureShareDisplayClusterStrings::Default::ShareName);
+		if (Object.IsValid())
+		{
+			ObjectProxy = Object->GetProxy();
+		}
+
+		if (IsActive())
+		{
+			// Initialize sync settings for nDisplay
+			FTextureShareCoreSyncSettings SyncSetting;
+			SyncSetting.FrameSyncSettings = Object->GetFrameSyncSettings(ETextureShareFrameSyncTemplate::DisplayCluster);
+
+			Object->SetSyncSetting(SyncSetting);
+
+			Object->BeginSession();
+
+			UE_LOG(LogTextureShareDisplayClusterPostProcess, Log, TEXT("TextureShareDisplayCluster: Initialized"));
+		}
+		else
+		{
+			ReleaseDisplayClusterPostProcessTextureShare();
+
+			UE_LOG(LogTextureShareDisplayClusterPostProcess, Error, TEXT("TextureShareDisplayCluster: Failed - initization failed"));
+		}
+	}
+}
+
+FTextureSharePostprocess::~FTextureSharePostprocess()
+{
+	if (Object.IsValid())
+	{
+		ReleaseDisplayClusterPostProcessTextureShare();
+	}
+}
+
+//////////////////////////////////////////////////////////////////////////////////////////////
+const FString& FTextureSharePostprocess::GetType() const
+{
+	static const FString Type(TextureSharePostprocessStrings::Postprocess::TextureShare);
+
+	return Type;
+}
+
+void FTextureSharePostprocess::ReleaseDisplayClusterPostProcessTextureShare()
+{
+	ObjectProxy.Reset();
+	Object.Reset();
+
+	TextureShareAPI().RemoveObject(TextureShareDisplayClusterStrings::Default::ShareName);
+}
+
+bool FTextureSharePostprocess::HandleStartScene(IDisplayClusterViewportManager* InViewportManager)
+{
+	return true;
+}
+
+void FTextureSharePostprocess::HandleEndScene(IDisplayClusterViewportManager* InViewportManager)
+{
+}
+
+//////////////////////////////////////////////////////////////////////////////////////////////
+void FTextureSharePostprocess::HandleSetupNewFrame(IDisplayClusterViewportManager* InViewportManager)
+{
+	if (IsActive() && Object->BeginFrameSync() && Object->IsFrameSyncActive())
+	{
+		// Update frame marker for current frame
+		Object->GetCoreData().FrameMarker.NextFrame();
+
+		// Share defined viewports from this DC node
+		UpdateSupportedViews(InViewportManager);
+
+		// Sync IPC data (read manual projection data)
+		if (Object->FrameSync(ETextureShareSyncStep::FramePreSetupBegin) && Object->IsFrameSyncActive())
+		{
+			// Update TS manual projection policy on this node
+			UpdateManualProjectionPolicy(InViewportManager);
+		}
+	}
+}
+
+void FTextureSharePostprocess::HandleBeginNewFrame(IDisplayClusterViewportManager* InViewportManager, FDisplayClusterRenderFrame& InOutRenderFrame)
+{
+	if (IsActive() && Object->IsFrameSyncActive())
+	{
+		if (Object->FrameSync(ETextureShareSyncStep::FrameSetupBegin) && Object->IsFrameSyncActive())
+		{
+			// Register viewport mapping
+			UpdateViews(InViewportManager);
+
+			Object->EndFrameSync(GetDisplayViewport(InViewportManager));
+
+			// Immediatelly begin proxy frame
+			ENQUEUE_RENDER_COMMAND(DisplayClusterPostProcessTextureShare_UpdateObjectProxy)(
+				[ObjectProxyRef = ObjectProxy.ToSharedRef()](FRHICommandListImmediate& RHICmdList)
+				{
+					ObjectProxyRef->BeginFrameSync_RenderThread(RHICmdList);
+				});
+		}
+	}
+}
+
+//////////////////////////////////////////////////////////////////////////////////////////////
+void FTextureSharePostprocess::HandleRenderFrameSetup_RenderThread(FRHICommandListImmediate& RHICmdList, const IDisplayClusterViewportManagerProxy* InViewportManagerProxy)
+{
+	check(IsInRenderingThread());
+
+	if (IsActive() && ObjectProxy->IsFrameSyncActive_RenderThread())
+	{
+		// Share RTT with remote process
+		ShareViewport_RenderThread(RHICmdList, InViewportManagerProxy, ETextureShareSyncStep::FrameProxyPreRenderEnd, EDisplayClusterViewportResourceType::InternalRenderTargetResource, TextureShareDisplayClusterStrings::Viewport::FinalColor);
+
+		ObjectProxy->FrameSync_RenderThread(RHICmdList, ETextureShareSyncStep::FrameProxyPreRenderEnd);
+	}
+}
+
+void FTextureSharePostprocess::HandleBeginUpdateFrameResources_RenderThread(FRHICommandListImmediate& RHICmdList, const IDisplayClusterViewportManagerProxy* InViewportManagerProxy)
+{
+	check(IsInRenderingThread());
+
+	if (IsActive() && ObjectProxy->IsFrameSyncActive_RenderThread())
+	{
+		// Share RTT with remote process
+		ShareViewport_RenderThread(RHICmdList, InViewportManagerProxy, ETextureShareSyncStep::FrameProxyRenderEnd, EDisplayClusterViewportResourceType::InputShaderResource, TextureShareDisplayClusterStrings::Viewport::Input);
+		ShareViewport_RenderThread(RHICmdList, InViewportManagerProxy, ETextureShareSyncStep::FrameProxyRenderEnd, EDisplayClusterViewportResourceType::MipsShaderResource, TextureShareDisplayClusterStrings::Viewport::Mips);
+
+		ObjectProxy->FrameSync_RenderThread(RHICmdList, ETextureShareSyncStep::FrameProxyRenderEnd);
+	}
+}
+
+void FTextureSharePostprocess::HandleUpdateFrameResourcesAfterWarpBlend_RenderThread(FRHICommandListImmediate& RHICmdList, const IDisplayClusterViewportManagerProxy* InViewportManagerProxy)
+{
+	check(IsInRenderingThread());
+
+	if (IsActive() && ObjectProxy->IsFrameSyncActive_RenderThread())
+	{
+		// Share RTT with remote process
+		ShareViewport_RenderThread(RHICmdList, InViewportManagerProxy, ETextureShareSyncStep::FrameProxyPostWarpEnd, EDisplayClusterViewportResourceType::InputShaderResource, TextureShareDisplayClusterStrings::Viewport::Warped, true);
+
+		ObjectProxy->FrameSync_RenderThread(RHICmdList, ETextureShareSyncStep::FrameProxyPostWarpEnd);
+	}
+}
+
+void FTextureSharePostprocess::HandleEndUpdateFrameResources_RenderThread(FRHICommandListImmediate& RHICmdList, const IDisplayClusterViewportManagerProxy* InViewportManagerProxy)
+{
+	check(IsInRenderingThread());
+
+	if (IsActive() && ObjectProxy->IsFrameSyncActive_RenderThread())
+	{
+		// Share RTT with remote process
+		ShareViewport_RenderThread(RHICmdList, InViewportManagerProxy, ETextureShareSyncStep::FrameProxyPostRenderEnd, EDisplayClusterViewportResourceType::OutputFrameTargetableResource,  TextureShareDisplayClusterStrings::Output::Viewport);
+		ShareFrame_RenderThread(RHICmdList, InViewportManagerProxy, ETextureShareSyncStep::FrameProxyPostRenderEnd, EDisplayClusterViewportResourceType::OutputFrameTargetableResource,     TextureShareDisplayClusterStrings::Output::Backbuffer);
+		ShareFrame_RenderThread(RHICmdList, InViewportManagerProxy, ETextureShareSyncStep::FrameProxyPostRenderEnd, EDisplayClusterViewportResourceType::AdditionalFrameTargetableResource, TextureShareDisplayClusterStrings::Output::BackbufferTemp);
+
+		ObjectProxy->EndFrameSync_RenderThread(RHICmdList);
+	}
+}
diff --git a/Engine/Plugins/VirtualProduction/TextureShare/Source/TextureShareDisplayCluster/Private/PostProcess/TextureSharePostprocess.h b/Engine/Plugins/VirtualProduction/TextureShare/Source/TextureShareDisplayCluster/Private/PostProcess/TextureSharePostprocess.h
new file mode 100644
--- /dev/null
+++ b/Engine/Plugins/VirtualProduction/TextureShare/Source/TextureShareDisplayCluster/Private/PostProcess/TextureSharePostprocess.h
@@ -0,0 +1,68 @@
+// Copyright Epic Games, Inc. All Rights Reserved.
+
+#pragma once
+#include "PostProcess/TextureSharePostprocessBase.h"
+#include "Render/Viewport/Containers/DisplayClusterViewport_Enums.h"
+#include "Containers/TextureShareContainers.h"
+
+#include "RHI.h"
+#include "RHICommandList.h"
+#include "RHIResources.h"
+
+class FDisplayClusterRenderFrame;
+
+class IDisplayClusterViewportManager;
+class IDisplayClusterViewportManagerProxy;
+
+/**
+ * TextureShare over nDisplay post-process interface
+ */
+class FTextureSharePostprocess
+	: public FTextureSharePostprocessBase
+{
+public:
+	FTextureSharePostprocess(const FString& PostprocessId, const struct FDisplayClusterConfigurationPostprocess* InConfigurationPostprocess);
+	virtual ~FTextureSharePostprocess();
+
+	virtual const FString& GetType() const override;
+
+	// Handle scene loading, viewport re-new etc
+	virtual bool HandleStartScene(IDisplayClusterViewportManager* InViewportManager) override;
+	virtual void HandleEndScene(IDisplayClusterViewportManager* InViewportManager) override;
+
+	// Handle frame on game thread
+	virtual void HandleSetupNewFrame(IDisplayClusterViewportManager* InViewportManager) override;
+	virtual void HandleBeginNewFrame(IDisplayClusterViewportManager* InViewportManager, FDisplayClusterRenderFrame& InOutRenderFrame) override;
+
+	// Handle frame on rendering thread
+	virtual void HandleRenderFrameSetup_RenderThread(FRHICommandListImmediate& RHICmdList, const IDisplayClusterViewportManagerProxy* InViewportManagerProxy) override;
+	virtual void HandleBeginUpdateFrameResources_RenderThread(FRHICommandListImmediate& RHICmdList, const IDisplayClusterViewportManagerProxy* InViewportManagerProxy) override;
+	virtual void HandleUpdateFrameResourcesAfterWarpBlend_RenderThread(FRHICommandListImmediate& RHICmdList, const IDisplayClusterViewportManagerProxy* InViewportManagerProxy) override;
+	virtual void HandleEndUpdateFrameResources_RenderThread(FRHICommandListImmediate& RHICmdList, const IDisplayClusterViewportManagerProxy* InViewportManagerProxy) override;
+
+	// Also share backbuffer
+	virtual bool IsPostProcessFrameAfterWarpBlendRequired() const override
+	{
+		return true;
+	}
+
+private:
+	bool IsActive() const
+	{
+		return Object.IsValid() && ObjectProxy.IsValid();
+	}
+
+	void ReleaseDisplayClusterPostProcessTextureShare();
+
+	void UpdateSupportedViews(IDisplayClusterViewportManager* InViewportManager);
+	void UpdateManualProjectionPolicy(IDisplayClusterViewportManager* InViewportManager);
+	void UpdateViews(IDisplayClusterViewportManager* InViewportManager);
+
+private:
+	void ShareViewport_RenderThread(FRHICommandListImmediate& RHICmdList, const IDisplayClusterViewportManagerProxy* InViewportManagerProxy, const ETextureShareSyncStep InReceiveSyncStep, const EDisplayClusterViewportResourceType InResourceType, const FString& InTextureId, bool bAfterWarpBlend = false) const;
+	void    ShareFrame_RenderThread(FRHICommandListImmediate& RHICmdList, const IDisplayClusterViewportManagerProxy* InViewportManagerProxy, const ETextureShareSyncStep InReceiveSyncStep, const EDisplayClusterViewportResourceType InResourceType, const FString& TextureId) const;
+
+private:
+	TSharedPtr<class ITextureShareObject, ESPMode::ThreadSafe>      Object;
+	TSharedPtr<class ITextureShareObjectProxy, ESPMode::ThreadSafe> ObjectProxy;
+};
diff --git a/Engine/Plugins/VirtualProduction/TextureShare/Source/TextureShareDisplayCluster/Private/PostProcess/TextureSharePostprocessBase.cpp b/Engine/Plugins/VirtualProduction/TextureShare/Source/TextureShareDisplayCluster/Private/PostProcess/TextureSharePostprocessBase.cpp
new file mode 100644
--- /dev/null
+++ b/Engine/Plugins/VirtualProduction/TextureShare/Source/TextureShareDisplayCluster/Private/PostProcess/TextureSharePostprocessBase.cpp
@@ -0,0 +1,47 @@
+// Copyright Epic Games, Inc. All Rights Reserved.
+
+#include "PostProcess/TextureSharePostprocessBase.h"
+#include "PostProcess/TextureSharePostprocessStrings.h"
+#include "Module/TextureShareDisplayClusterLog.h"
+
+#include "Render/Viewport/IDisplayClusterViewportManager.h"
+#include "DisplayClusterConfigurationTypes.h"
+
+//////////////////////////////////////////////////////////////////////////////////////////////
+// FTextureSharePostprocessBase
+//////////////////////////////////////////////////////////////////////////////////////////////
+FTextureSharePostprocessBase::FTextureSharePostprocessBase(const FString& InProjectionPolicyId, const FDisplayClusterConfigurationPostprocess* InConfigurationPostprocess)
+	: PostprocessId(InProjectionPolicyId)
+{
+	Parameters.Append(InConfigurationPostprocess->Parameters);
+	Order = InConfigurationPostprocess->Order;
+}
+
+FTextureSharePostprocessBase::~FTextureSharePostprocessBase()
+{
+}
+
+//////////////////////////////////////////////////////////////////////////////////////////////
+bool FTextureSharePostprocessBase::IsConfigurationChanged(const struct FDisplayClusterConfigurationPostprocess* InConfigurationPostprocess) const
+{
+	if (InConfigurationPostprocess->Type.Compare(GetType(), ESearchCase::IgnoreCase) != 0)
+	{
+		return true;
+	}
+
+	if (InConfigurationPostprocess->Parameters.Num() != Parameters.Num()) {
+
+		return true;
+	}
+
+	for (const TPair<FString, FString>& NewParamIt : InConfigurationPostprocess->Parameters) {
+
+		const FString* CurrentValue = Parameters.Find(NewParamIt.Key);
+		if (CurrentValue==nullptr || CurrentValue->Compare(NewParamIt.Value, ESearchCase::IgnoreCase) != 0) {
+			return true;
+		}
+	}
+
+	// Parameters not changed
+	return false;
+}
diff --git a/Engine/Plugins/VirtualProduction/TextureShare/Source/TextureShareDisplayCluster/Private/PostProcess/TextureSharePostprocessBase.h b/Engine/Plugins/VirtualProduction/TextureShare/Source/TextureShareDisplayCluster/Private/PostProcess/TextureSharePostprocessBase.h
new file mode 100644
--- /dev/null
+++ b/Engine/Plugins/VirtualProduction/TextureShare/Source/TextureShareDisplayCluster/Private/PostProcess/TextureSharePostprocessBase.h
@@ -0,0 +1,37 @@
+// Copyright Epic Games, Inc. All Rights Reserved.
+
+#pragma once
+#include "Render/PostProcess/IDisplayClusterPostProcess.h"
+
+/**
+ * Base postprocess
+ */
+class FTextureSharePostprocessBase
+	: public IDisplayClusterPostProcess
+{
+public:
+	FTextureSharePostprocessBase(const FString& PostprocessId, const struct FDisplayClusterConfigurationPostprocess* InConfigurationPostprocess);
+	virtual ~FTextureSharePostprocessBase() = 0;
+
+	virtual const FString& GetId() const override
+	{
+		return PostprocessId;
+	}
+
+	virtual int32 GetOrder() const override
+	{
+		return Order;
+	}
+
+	virtual const TMap<FString, FString>& GetParameters() const override
+	{
+		return Parameters;
+	}
+
+	virtual bool IsConfigurationChanged(const struct FDisplayClusterConfigurationPostprocess* InConfigurationPostprocess) const override;
+
+private:
+	FString PostprocessId;
+	TMap<FString, FString> Parameters;
+	int32 Order = -1;
+};
diff --git a/Engine/Plugins/VirtualProduction/TextureShare/Source/TextureShareDisplayCluster/Private/PostProcess/TextureSharePostprocessFactory.cpp b/Engine/Plugins/VirtualProduction/TextureShare/Source/TextureShareDisplayCluster/Private/PostProcess/TextureSharePostprocessFactory.cpp
new file mode 100644
--- /dev/null
+++ b/Engine/Plugins/VirtualProduction/TextureShare/Source/TextureShareDisplayCluster/Private/PostProcess/TextureSharePostprocessFactory.cpp
@@ -0,0 +1,19 @@
+// Copyright Epic Games, Inc. All Rights Reserved.
+
+#include "PostProcess/TextureSharePostprocessFactory.h"
+#include "PostProcess/TextureSharePostprocess.h"
+
+#include "Module/TextureShareDisplayClusterLog.h"
+#include "DisplayClusterConfigurationTypes_Base.h"
+
+//////////////////////////////////////////////////////////////////////////////////////////////
+// FTextureSharePostprocessFactory
+//////////////////////////////////////////////////////////////////////////////////////////////
+TSharedPtr<IDisplayClusterPostProcess, ESPMode::ThreadSafe> FTextureSharePostprocessFactory::Create(const FString& InPostProcessId, const FDisplayClusterConfigurationPostprocess* InConfigurationPostProcess)
+{
+	check(InConfigurationPostProcess != nullptr);
+
+	UE_LOG(LogTextureShareDisplayClusterPostProcess, Verbose, TEXT("Instantiating postprocess <%s> id='%s'"), *InConfigurationPostProcess->Type, *InPostProcessId);
+
+	return  MakeShared<FTextureSharePostprocess, ESPMode::ThreadSafe>(InPostProcessId, InConfigurationPostProcess);
+}
diff --git a/Engine/Plugins/VirtualProduction/TextureShare/Source/TextureShareDisplayCluster/Private/PostProcess/TextureSharePostprocessFactory.h b/Engine/Plugins/VirtualProduction/TextureShare/Source/TextureShareDisplayCluster/Private/PostProcess/TextureSharePostprocessFactory.h
new file mode 100644
--- /dev/null
+++ b/Engine/Plugins/VirtualProduction/TextureShare/Source/TextureShareDisplayCluster/Private/PostProcess/TextureSharePostprocessFactory.h
@@ -0,0 +1,16 @@
+// Copyright Epic Games, Inc. All Rights Reserved.
+
+#pragma once
+
+#include "Render/PostProcess/IDisplayClusterPostProcessFactory.h"
+
+/**
+ * Postprocess factory for TextureShare over nDisplay
+ */
+class FTextureSharePostprocessFactory
+	: public IDisplayClusterPostProcessFactory
+{
+public:
+	// Constructor for postprocess
+	virtual TSharedPtr<IDisplayClusterPostProcess, ESPMode::ThreadSafe> Create(const FString& PostProcessId, const struct FDisplayClusterConfigurationPostprocess* InConfigurationPostProcess) override;
+};
diff --git a/Engine/Plugins/VirtualProduction/TextureShare/Source/TextureShareDisplayCluster/Private/PostProcess/TextureSharePostprocessStrings.h b/Engine/Plugins/VirtualProduction/TextureShare/Source/TextureShareDisplayCluster/Private/PostProcess/TextureSharePostprocessStrings.h
new file mode 100644
--- /dev/null
+++ b/Engine/Plugins/VirtualProduction/TextureShare/Source/TextureShareDisplayCluster/Private/PostProcess/TextureSharePostprocessStrings.h
@@ -0,0 +1,12 @@
+// Copyright Epic Games, Inc. All Rights Reserved.
+
+#pragma once
+#include "CoreMinimal.h"
+
+namespace TextureSharePostprocessStrings
+{
+	namespace Postprocess
+	{
+		static constexpr auto TextureShare = TEXT("TextureShare");
+	}
+};
diff --git a/Engine/Plugins/VirtualProduction/TextureShare/Source/TextureShareDisplayCluster/Private/PostProcess/TextureSharePostprocess_Frame.cpp b/Engine/Plugins/VirtualProduction/TextureShare/Source/TextureShareDisplayCluster/Private/PostProcess/TextureSharePostprocess_Frame.cpp
new file mode 100644
--- /dev/null
+++ b/Engine/Plugins/VirtualProduction/TextureShare/Source/TextureShareDisplayCluster/Private/PostProcess/TextureSharePostprocess_Frame.cpp
@@ -0,0 +1,101 @@
+// Copyright Epic Games, Inc. All Rights Reserved.
+
+#include "PostProcess/TextureSharePostprocess.h"
+
+#include "PostProcess/TextureSharePostprocessStrings.h"
+#include "Projection/TextureShareProjectionStrings.h"
+#include "Misc/TextureShareDisplayClusterStrings.h"
+
+#include "Module/TextureShareDisplayClusterLog.h"
+
+#include "Containers/TextureShareCoreEnums.h"
+
+#include "ITextureShare.h"
+#include "ITextureShareObject.h"
+#include "ITextureShareDisplayCluster.h"
+#include "ITextureShareDisplayClusterAPI.h"
+
+#include "Render/Viewport/IDisplayClusterViewportManager.h"
+#include "Render/Viewport/IDisplayClusterViewport.h"
+#include "Render/Projection/IDisplayClusterProjectionPolicy.h"
+
+//////////////////////////////////////////////////////////////////////////////////////////////
+// FTextureSharePostprocess
+//////////////////////////////////////////////////////////////////////////////////////////////
+void FTextureSharePostprocess::UpdateSupportedViews(IDisplayClusterViewportManager* InViewportManager)
+{
+	if (InViewportManager)
+	{
+		// Get all existing viewports on DC node
+		for (const IDisplayClusterViewport* ViewportIt : InViewportManager->GetViewports())
+		{
+			if (ViewportIt)
+			{
+				const TArray<FDisplayClusterViewport_Context>& Contexts = ViewportIt->GetContexts();
+				for (const FDisplayClusterViewport_Context& ViewportContextIt : Contexts)
+				{
+					const ETextureShareEyeType EyeType = (Contexts.Num() == 1) ? ETextureShareEyeType::Default : (ViewportContextIt.ContextNum == 0 ? ETextureShareEyeType::StereoLeft : ETextureShareEyeType::StereoRight);
+					Object->GetCoreData().SupportedViews.Add(FTextureShareCoreViewDesc(ViewportIt->GetId(), EyeType));
+				}
+			}
+		}
+	}
+}
+
+void FTextureSharePostprocess::UpdateManualProjectionPolicy(IDisplayClusterViewportManager* InViewportManager)
+{
+	if (InViewportManager)
+	{
+		for (IDisplayClusterViewport* ViewportIt : InViewportManager->GetViewports())
+		{
+			TSharedPtr<IDisplayClusterProjectionPolicy, ESPMode::ThreadSafe> PrjPolicy = ViewportIt->GetProjectionPolicy();
+			if (PrjPolicy.IsValid() && PrjPolicy->GetType() == TextureShareProjectionStrings::Projection::TextureShare)
+			{
+				// Search the custom manual projection for viewport with the "TextureShare" projection policy type 
+				TArray<FTextureShareCoreManualProjection> ManualProjections;
+				for (const FTextureShareCoreObjectData& FrameDataIt : Object->GetReceivedCoreObjectData())
+				{
+					if (FrameDataIt.Data.ManualProjections.GetValuesByEqualsFunc(ViewportIt->GetId(), ManualProjections))
+					{
+						// Apply manual projection
+						static ITextureShareDisplayClusterAPI& TextureShareDisplayClusterAPI = ITextureShareDisplayCluster::Get().GetTextureShareDisplayClusterAPI();
+						TextureShareDisplayClusterAPI.TextureSharePolicySetProjectionData(PrjPolicy, ManualProjections);
+
+						// Save used manual projections sources
+						for (const FTextureShareCoreManualProjection& It : ManualProjections)
+						{
+							Object->GetCoreData().ManualProjectionsSources.Add(FTextureShareCoreManualProjectionSource(It.ViewDesc, FTextureShareCoreObjectFrameMarker(FrameDataIt.Desc, FrameDataIt.Data.FrameMarker)));
+						}
+						break;
+					}
+				}
+			}
+		}
+	}
+}
+
+void FTextureSharePostprocess::UpdateViews(IDisplayClusterViewportManager* InViewportManager)
+{
+	if (InViewportManager)
+	{
+		//Register all viewport for scene eye data capture by default
+		for (IDisplayClusterViewport* ViewportIt : InViewportManager->GetViewports())
+		{
+			if (ViewportIt)
+			{
+				const bool bMonoscopic = ViewportIt->GetContexts().Num() == 1;
+				for (const FDisplayClusterViewport_Context& ContextIt : ViewportIt->GetContexts())
+				{
+					const ETextureShareEyeType EyeType = bMonoscopic
+						? ETextureShareEyeType::Default
+						: ((ContextIt.ContextNum == 0) ? ETextureShareEyeType::StereoLeft : ETextureShareEyeType::StereoRight);
+
+					const FTextureShareCoreViewDesc InViewDesc(ViewportIt->GetId(), EyeType);
+
+					// Register all nDisplay viewports
+					Object->GetData().Views.Add(InViewDesc, ContextIt.StereoViewIndex, ContextIt.StereoscopicPass);
+				}
+			}
+		}
+	}
+}
diff --git a/Engine/Plugins/VirtualProduction/TextureShare/Source/TextureShareDisplayCluster/Private/PostProcess/TextureSharePostprocess_FrameProxy.cpp b/Engine/Plugins/VirtualProduction/TextureShare/Source/TextureShareDisplayCluster/Private/PostProcess/TextureSharePostprocess_FrameProxy.cpp
new file mode 100644
--- /dev/null
+++ b/Engine/Plugins/VirtualProduction/TextureShare/Source/TextureShareDisplayCluster/Private/PostProcess/TextureSharePostprocess_FrameProxy.cpp
@@ -0,0 +1,123 @@
+// Copyright Epic Games, Inc. All Rights Reserved.
+
+#include "PostProcess/TextureSharePostprocess.h"
+
+#include "PostProcess/TextureSharePostprocessStrings.h"
+#include "Misc/TextureShareDisplayClusterStrings.h"
+
+#include "Module/TextureShareDisplayClusterLog.h"
+#include "Containers/TextureShareCoreEnums.h"
+
+#include "ITextureShareObjectProxy.h"
+
+#include "Render/Viewport/IDisplayClusterViewportManagerProxy.h"
+#include "Render/Viewport/IDisplayClusterViewportProxy.h"
+#include "Render/Projection/IDisplayClusterProjectionPolicy.h"
+
+//////////////////////////////////////////////////////////////////////////////////////////////
+namespace DisplayClusterPostProcessTextureShareHelpers
+{
+	// Support warp blend logic
+	static bool ShouldApplyWarpBlend(IDisplayClusterViewportProxy* ViewportProxy)
+	{
+		if (ViewportProxy->GetPostRenderSettings_RenderThread().Replace.IsEnabled())
+		{
+			// When used override texture, disable warp blend
+			return false;
+		}
+
+		const TSharedPtr<IDisplayClusterProjectionPolicy, ESPMode::ThreadSafe>& PrjPolicy = ViewportProxy->GetProjectionPolicy_RenderThread();
+
+		// Projection policy must support warp blend op
+		return PrjPolicy.IsValid() && PrjPolicy->IsWarpBlendSupported();
+	}
+};
+
+using namespace DisplayClusterPostProcessTextureShareHelpers;
+
+//////////////////////////////////////////////////////////////////////////////////////////////
+// FTextureSharePostprocess
+//////////////////////////////////////////////////////////////////////////////////////////////// Copyright Epic Games, Inc. All Rights Reserved.
+void FTextureSharePostprocess::ShareViewport_RenderThread(FRHICommandListImmediate& RHICmdList, const IDisplayClusterViewportManagerProxy* InViewportManagerProxy, const ETextureShareSyncStep InReceiveSyncStep, const EDisplayClusterViewportResourceType InResourceType, const FString& InTextureId, bool bAfterWarpBlend) const
+{
+	if (InViewportManagerProxy)
+	{
+		for (IDisplayClusterViewportProxy* ViewportProxyIt : InViewportManagerProxy->GetViewports_RenderThread())
+		{
+			if (ViewportProxyIt)
+			{
+				// Get viewport resource type
+				EDisplayClusterViewportResourceType ResourceType = InResourceType;
+				if (bAfterWarpBlend && ShouldApplyWarpBlend(ViewportProxyIt))
+				{
+					const TSharedPtr<IDisplayClusterProjectionPolicy, ESPMode::ThreadSafe>& PrjPolicy = ViewportProxyIt->GetProjectionPolicy_RenderThread();
+					if (PrjPolicy->ShouldUseAdditionalTargetableResource())
+					{
+						TArray<FRHITexture2D*> Resources;
+						if (ViewportProxyIt->GetResources_RenderThread(EDisplayClusterViewportResourceType::AdditionalTargetableResource, Resources) && Resources.Num())
+						{
+							ResourceType = EDisplayClusterViewportResourceType::AdditionalTargetableResource;
+						}
+					}
+				}
+
+				const bool bMonoscopic = ViewportProxyIt->GetContexts_RenderThread().Num() == 1;
+
+				TArray<FRHITexture2D*> ViewportResources;
+				TArray<FIntRect> ViewportResourceRects;
+				if (ViewportProxyIt->GetResourcesWithRects_RenderThread(ResourceType, ViewportResources, ViewportResourceRects))
+				{
+					for (int32 ContextIndex = 0; ContextIndex < ViewportResources.Num(); ContextIndex++)
+					{
+						const int32 InGPUIndex = ViewportProxyIt->GetContexts_RenderThread()[ContextIndex].GPUIndex;
+
+						const ETextureShareEyeType EyeType = bMonoscopic
+							? ETextureShareEyeType::Default
+							: ((ContextIndex == 0) ? ETextureShareEyeType::StereoLeft : ETextureShareEyeType::StereoRight);
+
+						const FTextureShareCoreViewDesc InViewDesc(ViewportProxyIt->GetId(), EyeType);
+
+						// Send
+						ObjectProxy->ShareResource_RenderThread(RHICmdList, FTextureShareCoreResourceDesc(InTextureId, InViewDesc, ETextureShareTextureOp::Read), ViewportResources[ContextIndex], InGPUIndex, &ViewportResourceRects[ContextIndex]);
+						// Delayed Receive
+						ObjectProxy->ShareResource_RenderThread(RHICmdList, FTextureShareCoreResourceDesc(InTextureId, InViewDesc, ETextureShareTextureOp::Write, InReceiveSyncStep), ViewportResources[ContextIndex], InGPUIndex, &ViewportResourceRects[ContextIndex]);
+					}
+				}
+			}
+		}
+	}
+}
+
+void FTextureSharePostprocess::ShareFrame_RenderThread(FRHICommandListImmediate& RHICmdList, const IDisplayClusterViewportManagerProxy* InViewportManagerProxy, const ETextureShareSyncStep InReceiveSyncStep, const EDisplayClusterViewportResourceType InResourceType, const FString& InTextureId) const
+{
+	if (InViewportManagerProxy)
+	{
+		TArray<FRHITexture2D*> FrameResources;
+		TArray<FRHITexture2D*> AdditionalFrameResources;
+		TArray<FIntPoint> TargetOffsets;
+		if (InViewportManagerProxy->GetFrameTargets_RenderThread(FrameResources, TargetOffsets, &AdditionalFrameResources))
+		{
+			TArray<FRHITexture2D*>& SharedResources = (InResourceType == EDisplayClusterViewportResourceType::AdditionalFrameTargetableResource)
+				? AdditionalFrameResources : FrameResources;
+
+			const bool bMonoscopic = SharedResources.Num() == 1;
+
+			for (int32 ContextIndex = 0; ContextIndex < SharedResources.Num(); ContextIndex++)
+			{
+				const int32 InGPUIndex = -1;
+
+				const ETextureShareEyeType EyeType = bMonoscopic
+					? ETextureShareEyeType::Default
+					: ((ContextIndex == 0) ? ETextureShareEyeType::StereoLeft : ETextureShareEyeType::StereoRight);
+
+				const FTextureShareCoreViewDesc InViewDesc(EyeType);
+
+				// Send
+				ObjectProxy->ShareResource_RenderThread(RHICmdList, FTextureShareCoreResourceDesc(InTextureId, InViewDesc, ETextureShareTextureOp::Read), SharedResources[ContextIndex], InGPUIndex);
+
+				// Receive
+				ObjectProxy->ShareResource_RenderThread(RHICmdList, FTextureShareCoreResourceDesc(InTextureId, InViewDesc, ETextureShareTextureOp::Write, InReceiveSyncStep), SharedResources[ContextIndex], InGPUIndex);
+			}
+		}
+	}
+}
diff --git a/Engine/Plugins/VirtualProduction/TextureShare/Source/TextureShareDisplayCluster/Private/Projection/DisplayClusterProjectionPolicyBase.cpp b/Engine/Plugins/VirtualProduction/TextureShare/Source/TextureShareDisplayCluster/Private/Projection/DisplayClusterProjectionPolicyBase.cpp
deleted file mode 100644
diff --git a/Engine/Plugins/VirtualProduction/TextureShare/Source/TextureShareDisplayCluster/Private/Projection/DisplayClusterProjectionPolicyBase.h b/Engine/Plugins/VirtualProduction/TextureShare/Source/TextureShareDisplayCluster/Private/Projection/DisplayClusterProjectionPolicyBase.h
deleted file mode 100644
diff --git a/Engine/Plugins/VirtualProduction/TextureShare/Source/TextureShareDisplayCluster/Private/Projection/DisplayClusterProjectionStrings.h b/Engine/Plugins/VirtualProduction/TextureShare/Source/TextureShareDisplayCluster/Private/Projection/DisplayClusterProjectionStrings.h
deleted file mode 100644
diff --git a/Engine/Plugins/VirtualProduction/TextureShare/Source/TextureShareDisplayCluster/Private/Projection/TextureShare/DisplayClusterProjectionTextureSharePolicy.cpp b/Engine/Plugins/VirtualProduction/TextureShare/Source/TextureShareDisplayCluster/Private/Projection/TextureShare/DisplayClusterProjectionTextureSharePolicy.cpp
deleted file mode 100644
diff --git a/Engine/Plugins/VirtualProduction/TextureShare/Source/TextureShareDisplayCluster/Private/Projection/TextureShare/DisplayClusterProjectionTextureSharePolicy.h b/Engine/Plugins/VirtualProduction/TextureShare/Source/TextureShareDisplayCluster/Private/Projection/TextureShare/DisplayClusterProjectionTextureSharePolicy.h
deleted file mode 100644
diff --git a/Engine/Plugins/VirtualProduction/TextureShare/Source/TextureShareDisplayCluster/Private/Projection/TextureShare/DisplayClusterProjectionTextureSharePolicyFactory.cpp b/Engine/Plugins/VirtualProduction/TextureShare/Source/TextureShareDisplayCluster/Private/Projection/TextureShare/DisplayClusterProjectionTextureSharePolicyFactory.cpp
deleted file mode 100644
diff --git a/Engine/Plugins/VirtualProduction/TextureShare/Source/TextureShareDisplayCluster/Private/Projection/TextureShare/DisplayClusterProjectionTextureSharePolicyFactory.h b/Engine/Plugins/VirtualProduction/TextureShare/Source/TextureShareDisplayCluster/Private/Projection/TextureShare/DisplayClusterProjectionTextureSharePolicyFactory.h
deleted file mode 100644
diff --git a/Engine/Plugins/VirtualProduction/TextureShare/Source/TextureShareDisplayCluster/Private/Projection/TextureShareProjectionPolicy.cpp b/Engine/Plugins/VirtualProduction/TextureShare/Source/TextureShareDisplayCluster/Private/Projection/TextureShareProjectionPolicy.cpp
new file mode 100644
--- /dev/null
+++ b/Engine/Plugins/VirtualProduction/TextureShare/Source/TextureShareDisplayCluster/Private/Projection/TextureShareProjectionPolicy.cpp
@@ -0,0 +1,175 @@
+// Copyright Epic Games, Inc. All Rights Reserved.
+
+#include "Projection/TextureShareProjectionPolicy.h"
+#include "Projection/TextureShareProjectionStrings.h"
+
+#include "Module/TextureShareDisplayClusterLog.h"
+
+#include "Game/IDisplayClusterGameManager.h"
+
+#include "Render/Viewport/IDisplayClusterViewport.h"
+#include "Render/Viewport/IDisplayClusterViewportManager.h"
+
+#include "Containers/TextureShareCoreContainers.h"
+
+//////////////////////////////////////////////////////////////////////////////////////////////
+// FTextureShareProjectionPolicy
+//////////////////////////////////////////////////////////////////////////////////////////////
+FTextureShareProjectionPolicy::FTextureShareProjectionPolicy(const FString& ProjectionPolicyId, const struct FDisplayClusterConfigurationProjection* InConfigurationProjectionPolicy)
+	: FTextureShareProjectionPolicyBase(ProjectionPolicyId, InConfigurationProjectionPolicy)
+{
+}
+
+FTextureShareProjectionPolicy::~FTextureShareProjectionPolicy()
+{
+}
+
+//////////////////////////////////////////////////////////////////////////////////////////////
+const FString& FTextureShareProjectionPolicy::GetType() const
+{
+	static const FString Type(TextureShareProjectionStrings::Projection::TextureShare);
+
+	return Type;
+}
+
+bool FTextureShareProjectionPolicy::SetCustomProjection(const TArray<FTextureShareCoreManualProjection>& InProjectionData)
+{
+	check(IsInGameThread());
+
+	if (InProjectionData.Num() > 2 || (InProjectionData.Num() == 0))
+	{
+		return false;
+	}
+
+	// Sort eyes to [2]: <Mono-Mono> or <Left-Right>
+	Projections.Empty();
+	Projections.AddDefaulted(2);
+
+	for (const FTextureShareCoreManualProjection& It : InProjectionData)
+	{
+		switch (It.ViewDesc.EyeType)
+		{
+		case ETextureShareEyeType::StereoLeft:
+			Projections[0] = It;
+			break;
+
+		case ETextureShareEyeType::StereoRight:
+			Projections[1] = It;
+			break;
+
+		case ETextureShareEyeType::Default:
+		default:
+			if (InProjectionData.Num() != 1)
+			{
+				// bad input
+				return false;
+			}
+
+			Projections[0] = Projections[1] = It;
+			break;
+		}
+	}
+
+	return true;
+}
+
+//////////////////////////////////////////////////////////////////////////////////////////////
+bool FTextureShareProjectionPolicy::HandleStartScene(class IDisplayClusterViewport* InViewport)
+{
+	return true;
+}
+
+void FTextureShareProjectionPolicy::HandleEndScene(class IDisplayClusterViewport* InViewport)
+{
+
+}
+
+bool FTextureShareProjectionPolicy::CalculateView(class IDisplayClusterViewport* InViewport, const uint32 InContextNum, FVector& InOutViewLocation, FRotator& InOutViewRotation, const FVector& ViewOffset, const float WorldToMeters, const float InNCP, const float InFCP)
+{
+	check(IsInGameThread());
+
+	if (Projections.Num() == 0)
+	{
+		// view data not assigned
+		return false;
+	}
+
+	check((int32)InContextNum < Projections.Num());
+
+	const FTextureShareCoreManualProjection& Src = Projections[InContextNum];
+
+	// Add local rotation specified in config
+	switch (Src.ViewLocationType)
+	{
+	default:
+	case ETextureShareViewLocationDataType::Original:
+		break;
+	case ETextureShareViewLocationDataType::Relative:
+		InOutViewLocation += Src.ViewLocation;
+		break;
+	case ETextureShareViewLocationDataType::Absolute:
+		InOutViewLocation = Src.ViewLocation + ViewOffset;
+		break;
+	case ETextureShareViewLocationDataType::Original_NoViewOffset:
+		InOutViewLocation -= ViewOffset;
+		break;
+	case ETextureShareViewLocationDataType::Relative_NoViewOffset:
+		InOutViewLocation -= ViewOffset;
+		InOutViewLocation += Src.ViewLocation;
+		break;
+	case ETextureShareViewLocationDataType::Absolute_NoViewOffset:
+		InOutViewLocation = Src.ViewLocation;
+		break;
+	};
+
+	switch (Src.ViewRotationType)
+	{
+	default:
+	case ETextureShareViewRotationDataType::Original:
+		break;
+	case ETextureShareViewRotationDataType::Relative:
+		InOutViewRotation += Src.ViewRotation;
+		break;
+	case ETextureShareViewRotationDataType::Absolute:
+		InOutViewRotation = Src.ViewRotation;
+		break;
+	}
+
+	// Store culling data
+	NCP = InNCP;
+	FCP = InFCP;
+
+	return true;
+}
+
+bool FTextureShareProjectionPolicy::GetProjectionMatrix(IDisplayClusterViewport* InViewport, const uint32 InContextNum, FMatrix& OutPrjMatrix)
+{
+	check(IsInGameThread());
+
+	if (Projections.Num() == 0)
+	{
+		// view data not assigned
+		return false;
+	}
+
+	check((int32)InContextNum < Projections.Num());
+
+	const FTextureShareCoreManualProjection& Src = Projections[InContextNum];
+
+	switch (Src.ProjectionType)
+	{
+	case ETextureShareCoreSceneViewManualProjectionType::Matrix:
+		OutPrjMatrix = Src.ProjectionMatrix;
+		return true;
+
+	case ETextureShareCoreSceneViewManualProjectionType::FrustumAngles:
+		InViewport->CalculateProjectionMatrix(InContextNum, Src.FrustumAngles.Left, Src.FrustumAngles.Right, Src.FrustumAngles.Top, Src.FrustumAngles.Bottom, NCP, FCP, true);
+		OutPrjMatrix = InViewport->GetContexts()[InContextNum].ProjectionMatrix;
+		return true;
+
+	default:
+		break;
+	}
+
+	return false;
+}
diff --git a/Engine/Plugins/VirtualProduction/TextureShare/Source/TextureShareDisplayCluster/Private/Projection/TextureShareProjectionPolicy.h b/Engine/Plugins/VirtualProduction/TextureShare/Source/TextureShareDisplayCluster/Private/Projection/TextureShareProjectionPolicy.h
new file mode 100644
--- /dev/null
+++ b/Engine/Plugins/VirtualProduction/TextureShare/Source/TextureShareDisplayCluster/Private/Projection/TextureShareProjectionPolicy.h
@@ -0,0 +1,43 @@
+// Copyright Epic Games, Inc. All Rights Reserved.
+
+#pragma once
+
+#include "Projection/TextureShareProjectionPolicyBase.h"
+
+struct FTextureShareCoreManualProjection;
+
+/**
+ * Manual projection policy for TextureShare
+ */
+class FTextureShareProjectionPolicy
+	: public FTextureShareProjectionPolicyBase
+{
+public:
+	FTextureShareProjectionPolicy(const FString& ProjectionPolicyId, const struct FDisplayClusterConfigurationProjection* InConfigurationProjectionPolicy);
+	virtual ~FTextureShareProjectionPolicy();
+
+public:
+	virtual const FString& GetType() const;
+
+	virtual bool HandleStartScene(class IDisplayClusterViewport* InViewport) override;
+	virtual void HandleEndScene(class IDisplayClusterViewport* InViewport) override;
+
+	virtual bool CalculateView(IDisplayClusterViewport* InViewport, const uint32 InContextNum, FVector& InOutViewLocation, FRotator& InOutViewRotation, const FVector& ViewOffset, const float WorldToMeters, const float NCP, const float FCP) override;
+	virtual bool GetProjectionMatrix(IDisplayClusterViewport* InViewport, const uint32 InContextNum, FMatrix& OutPrjMatrix) override;
+
+	virtual bool IsWarpBlendSupported() override
+	{ return false; }
+
+	virtual bool ShouldUseSourceTextureWithMips() const override
+	{ return false; }
+
+public:
+	bool SetCustomProjection(const TArray<FTextureShareCoreManualProjection>& InProjectionData);
+
+private:
+	TArray<FTextureShareCoreManualProjection> Projections;
+	
+	// Near/far clip planes
+	float NCP;
+	float FCP;
+};
diff --git a/Engine/Plugins/VirtualProduction/TextureShare/Source/TextureShareDisplayCluster/Private/Projection/TextureShareProjectionPolicyBase.cpp b/Engine/Plugins/VirtualProduction/TextureShare/Source/TextureShareDisplayCluster/Private/Projection/TextureShareProjectionPolicyBase.cpp
new file mode 100644
--- /dev/null
+++ b/Engine/Plugins/VirtualProduction/TextureShare/Source/TextureShareDisplayCluster/Private/Projection/TextureShareProjectionPolicyBase.cpp
@@ -0,0 +1,91 @@
+// Copyright Epic Games, Inc. All Rights Reserved.
+
+#include "Projection/TextureShareProjectionPolicyBase.h"
+#include "Projection/TextureShareProjectionStrings.h"
+
+#include "IDisplayCluster.h"
+#include "Game/IDisplayClusterGameManager.h"
+
+#include "DisplayClusterConfigurationTypes_Base.h"
+
+#include "Render/Viewport/IDisplayClusterViewportManager.h"
+#include "Render/Viewport/IDisplayClusterViewport.h"
+
+//////////////////////////////////////////////////////////////////////////////////////////////
+// FTextureShareProjectionPolicyBase
+//////////////////////////////////////////////////////////////////////////////////////////////
+FTextureShareProjectionPolicyBase::FTextureShareProjectionPolicyBase(const FString& InProjectionPolicyId, const FDisplayClusterConfigurationProjection* InConfigurationProjectionPolicy)
+	: ProjectionPolicyId(InProjectionPolicyId)
+{
+	Parameters.Append(InConfigurationProjectionPolicy->Parameters);
+}
+
+FTextureShareProjectionPolicyBase::~FTextureShareProjectionPolicyBase()
+{
+}
+
+bool FTextureShareProjectionPolicyBase::IsEditorOperationMode_RenderThread(const IDisplayClusterViewportProxy* InViewportProxy)
+{
+	check(IsInRenderingThread());
+
+	if (const bool bIsClusterOperationMode = IDisplayCluster::Get().GetOperationMode() == EDisplayClusterOperationMode::Cluster)
+	{
+		return false;
+	}
+
+	return true;
+}
+
+bool FTextureShareProjectionPolicyBase::IsEditorOperationMode(class IDisplayClusterViewport* InViewport)
+{
+	check(IsInGameThread());
+
+	// Hide spam in logs when configuring VP in editor
+	if (const bool bIsClusterOperationMode = IDisplayCluster::Get().GetOperationMode() == EDisplayClusterOperationMode::Cluster)
+	{
+		return false;
+	}
+
+	// Get state from viewport world (UE-114493)
+	if (InViewport)
+	{
+		if (UWorld* CurrentWorld = InViewport->GetOwner().GetCurrentWorld())
+		{
+			if (CurrentWorld)
+			{
+				switch (CurrentWorld->WorldType)
+				{
+				case EWorldType::Editor:
+				case EWorldType::EditorPreview:
+					return true;
+
+				default:
+					break;
+				}
+			}
+		}
+	}
+
+	return false;
+}
+
+bool FTextureShareProjectionPolicyBase::IsConfigurationChanged(const struct FDisplayClusterConfigurationProjection* InConfigurationProjectionPolicy) const
+{
+	if (InConfigurationProjectionPolicy->Parameters.Num() != Parameters.Num())
+	{
+		return true;
+	}
+
+	for (const TPair<FString, FString>& NewParamIt : InConfigurationProjectionPolicy->Parameters)
+	{
+		const FString* CurrentValue = Parameters.Find(NewParamIt.Key);
+
+		if (CurrentValue==nullptr || CurrentValue->Compare(NewParamIt.Value, ESearchCase::IgnoreCase) != 0)
+		{
+			return true;
+		}
+	}
+
+	// Parameters not changed
+	return false;
+}
diff --git a/Engine/Plugins/VirtualProduction/TextureShare/Source/TextureShareDisplayCluster/Private/Projection/TextureShareProjectionPolicyBase.h b/Engine/Plugins/VirtualProduction/TextureShare/Source/TextureShareDisplayCluster/Private/Projection/TextureShareProjectionPolicyBase.h
new file mode 100644
--- /dev/null
+++ b/Engine/Plugins/VirtualProduction/TextureShare/Source/TextureShareDisplayCluster/Private/Projection/TextureShareProjectionPolicyBase.h
@@ -0,0 +1,36 @@
+// Copyright Epic Games, Inc. All Rights Reserved.
+
+#pragma once
+
+#include "Render/Projection/IDisplayClusterProjectionPolicy.h"
+
+/**
+* Base projection policy
+*/
+class FTextureShareProjectionPolicyBase
+	: public IDisplayClusterProjectionPolicy
+{
+public:
+	FTextureShareProjectionPolicyBase(const FString& ProjectionPolicyId, const struct FDisplayClusterConfigurationProjection* InConfigurationProjectionPolicy);
+	virtual ~FTextureShareProjectionPolicyBase() = 0;
+
+	virtual const FString& GetId() const override
+	{
+		return ProjectionPolicyId;
+	}
+
+	virtual const TMap<FString, FString>& GetParameters() const override
+	{
+		return Parameters;
+	}
+
+	virtual bool IsConfigurationChanged(const struct FDisplayClusterConfigurationProjection* InConfigurationProjectionPolicy) const override;
+
+	static bool IsEditorOperationMode(class IDisplayClusterViewport* InViewport);
+	static bool IsEditorOperationMode_RenderThread(const IDisplayClusterViewportProxy* InViewportProxy);
+
+private:
+	// Added 'Policy' prefix to avoid "... hides class name ..." warnings in child classes
+	FString ProjectionPolicyId;
+	TMap<FString, FString> Parameters;
+};
diff --git a/Engine/Plugins/VirtualProduction/TextureShare/Source/TextureShareDisplayCluster/Private/Projection/TextureShareProjectionPolicyFactory.cpp b/Engine/Plugins/VirtualProduction/TextureShare/Source/TextureShareDisplayCluster/Private/Projection/TextureShareProjectionPolicyFactory.cpp
new file mode 100644
--- /dev/null
+++ b/Engine/Plugins/VirtualProduction/TextureShare/Source/TextureShareDisplayCluster/Private/Projection/TextureShareProjectionPolicyFactory.cpp
@@ -0,0 +1,20 @@
+// Copyright Epic Games, Inc. All Rights Reserved.
+
+#include "Projection/TextureShareProjectionPolicyFactory.h"
+#include "Projection/TextureShareProjectionPolicy.h"
+
+#include "Module/TextureShareDisplayClusterLog.h"
+
+#include "DisplayClusterConfigurationTypes.h"
+
+//////////////////////////////////////////////////////////////////////////////////////////////
+// FTextureShareProjectionPolicyFactory
+//////////////////////////////////////////////////////////////////////////////////////////////
+TSharedPtr<IDisplayClusterProjectionPolicy, ESPMode::ThreadSafe> FTextureShareProjectionPolicyFactory::Create(const FString& ProjectionPolicyId, const struct FDisplayClusterConfigurationProjection* InConfigurationProjectionPolicy)
+{
+	check(InConfigurationProjectionPolicy != nullptr);
+
+	UE_LOG(LogTextureShareDisplayClusterProjection, Verbose, TEXT("Instantiating projection policy <%s> id='%s'"), *InConfigurationProjectionPolicy->Type, *ProjectionPolicyId);
+
+	return  MakeShared<FTextureShareProjectionPolicy, ESPMode::ThreadSafe>(ProjectionPolicyId, InConfigurationProjectionPolicy);
+}
diff --git a/Engine/Plugins/VirtualProduction/TextureShare/Source/TextureShareDisplayCluster/Private/Projection/TextureShareProjectionPolicyFactory.h b/Engine/Plugins/VirtualProduction/TextureShare/Source/TextureShareDisplayCluster/Private/Projection/TextureShareProjectionPolicyFactory.h
new file mode 100644
--- /dev/null
+++ b/Engine/Plugins/VirtualProduction/TextureShare/Source/TextureShareDisplayCluster/Private/Projection/TextureShareProjectionPolicyFactory.h
@@ -0,0 +1,16 @@
+// Copyright Epic Games, Inc. All Rights Reserved.
+
+#pragma once
+
+#include "Render/Projection/IDisplayClusterProjectionPolicyFactory.h"
+
+/**
+ * Projection policy factory for the 'TextureShare' policy
+ */
+class FTextureShareProjectionPolicyFactory
+	: public IDisplayClusterProjectionPolicyFactory
+{
+public:
+	// Constructor for policy
+	virtual TSharedPtr<IDisplayClusterProjectionPolicy, ESPMode::ThreadSafe> Create(const FString& ProjectionPolicyId, const struct FDisplayClusterConfigurationProjection* InConfigurationProjectionPolicy) override;
+};
diff --git a/Engine/Plugins/VirtualProduction/TextureShare/Source/TextureShareDisplayCluster/Private/Projection/TextureShareProjectionStrings.h b/Engine/Plugins/VirtualProduction/TextureShare/Source/TextureShareDisplayCluster/Private/Projection/TextureShareProjectionStrings.h
new file mode 100644
--- /dev/null
+++ b/Engine/Plugins/VirtualProduction/TextureShare/Source/TextureShareDisplayCluster/Private/Projection/TextureShareProjectionStrings.h
@@ -0,0 +1,12 @@
+// Copyright Epic Games, Inc. All Rights Reserved.
+
+#pragma once
+#include "CoreMinimal.h"
+
+namespace TextureShareProjectionStrings
+{
+	namespace Projection
+	{
+		static constexpr auto TextureShare = TEXT("textureshare");
+	}
+};
diff --git a/Engine/Plugins/VirtualProduction/TextureShare/Source/TextureShareDisplayCluster/TextureShareDisplayCluster.Build.cs b/Engine/Plugins/VirtualProduction/TextureShare/Source/TextureShareDisplayCluster/TextureShareDisplayCluster.Build.cs
--- a/Engine/Plugins/VirtualProduction/TextureShare/Source/TextureShareDisplayCluster/TextureShareDisplayCluster.Build.cs
+++ b/Engine/Plugins/VirtualProduction/TextureShare/Source/TextureShareDisplayCluster/TextureShareDisplayCluster.Build.cs
@@ -15,16 +15,33 @@ public class TextureShareDisplayCluster : ModuleRules
 			}
 		);
 
-		PublicDependencyModuleNames.AddRange(
+		// List of modules names (no path needed) with header files that our module's public headers needs access to, but we don't need to "import" or link against.
+		PublicIncludePathModuleNames.AddRange(
 			new string[]
 			{
-				"Core",
-				"CoreUObject",
 				"DisplayCluster",
 				"DisplayClusterShaders",
 				"DisplayClusterConfiguration",
+
+				// Module names for DisplayClusterConfiguration
+				"ActorLayerUtilities",
+				"CinematicCamera",
+				"MediaAssets",
+				"Media",
+				"MediaUtils",
+				"MediaIOCore",
+				"OpenColorIO",
+			});
+
+		// List of public dependency module names (no path needed) (automatically does the private/public include). These are modules that are required by our public source files.
+		PublicDependencyModuleNames.AddRange(
+			new string[]
+			{
+				"Core",
+				"CoreUObject",
 			});
 
+		// List of private dependency module names.  These are modules that our private code depends on but nothing in our public include files depend on.
 		PrivateDependencyModuleNames.AddRange(
 			new string[]
 			{
