TextureShare:
 - Re-add the nDisplay dependency to resolve the issue. (Independence between TextureShare and nDisplay will become possible in the future by refactoring the public nDisplay API into a separate plugin)
 - Default SDK process name can now use TEXTURESHARE_PROJECT_NAME macro as source. This allows you to run multiple TextureShareSDK example applications at the same time.

diff --git a/Engine/Plugins/VirtualProduction/TextureShare/Source/TextureShareCore/Private/Misc/TextureShareCoreStrings.h b/Engine/Plugins/VirtualProduction/TextureShare/Source/TextureShareCore/Private/Misc/TextureShareCoreStrings.h
--- a/Engine/Plugins/VirtualProduction/TextureShare/Source/TextureShareCore/Private/Misc/TextureShareCoreStrings.h
+++ b/Engine/Plugins/VirtualProduction/TextureShare/Source/TextureShareCore/Private/Misc/TextureShareCoreStrings.h
@@ -18,7 +18,12 @@ namespace TextureShareCoreStrings
 
 		namespace ProcessName
 		{
+#if defined TEXTURESHARE_PROJECT_NAME
+			// Use TextureShare project name as process name
+			static constexpr auto SDK = TEXT(TEXTURESHARE_PROJECT_NAME);
+#else
 			static constexpr auto SDK = TEXT("TextureShareSDK");
+#endif
 			static constexpr auto UE = TEXT("UnrealEngine");
 		}
 
diff --git a/Engine/Plugins/VirtualProduction/TextureShare/Source/TextureShareDisplayCluster/TextureShareDisplayCluster.Build.cs b/Engine/Plugins/VirtualProduction/TextureShare/Source/TextureShareDisplayCluster/TextureShareDisplayCluster.Build.cs
--- a/Engine/Plugins/VirtualProduction/TextureShare/Source/TextureShareDisplayCluster/TextureShareDisplayCluster.Build.cs
+++ b/Engine/Plugins/VirtualProduction/TextureShare/Source/TextureShareDisplayCluster/TextureShareDisplayCluster.Build.cs
@@ -15,30 +15,16 @@ public class TextureShareDisplayCluster : ModuleRules
 			}
 		);
 
-		// List of modules names (no path needed) with header files that our module's public headers needs access to, but we don't need to "import" or link against.
-		PublicIncludePathModuleNames.AddRange(
-			new string[]
-			{
-				"DisplayCluster",
-				"DisplayClusterShaders",
-				"DisplayClusterConfiguration",
-
-				// Module names for DisplayClusterConfiguration
-				"ActorLayerUtilities",
-				"CinematicCamera",
-				"MediaAssets",
-				"Media",
-				"MediaUtils",
-				"MediaIOCore",
-				"OpenColorIO",
-			});
-
 		// List of public dependency module names (no path needed) (automatically does the private/public include). These are modules that are required by our public source files.
 		PublicDependencyModuleNames.AddRange(
 			new string[]
 			{
 				"Core",
 				"CoreUObject",
+
+				"DisplayCluster",
+				"DisplayClusterShaders",
+				"DisplayClusterConfiguration",
 			});
 
 		// List of private dependency module names.  These are modules that our private code depends on but nothing in our public include files depend on.
diff --git a/Engine/Plugins/VirtualProduction/TextureShare/TextureShare.uplugin b/Engine/Plugins/VirtualProduction/TextureShare/TextureShare.uplugin
--- a/Engine/Plugins/VirtualProduction/TextureShare/TextureShare.uplugin
+++ b/Engine/Plugins/VirtualProduction/TextureShare/TextureShare.uplugin
@@ -39,5 +39,11 @@
 				"Win64"
 			]
 		}
+	],
+	"Plugins": [
+		{
+			"Name": "nDisplay",
+			"Enabled": true
+		}
 	]
 }
