TextureShare: fixed RDG issues

diff --git a/Engine/Extras/VirtualProduction/TextureShare/Samples/D3D11/DC5_CopyToBackbuffer/Source/D3D11TextureShareSample.cpp b/Engine/Extras/VirtualProduction/TextureShare/Samples/D3D11/DC5_CopyToBackbuffer/Source/D3D11TextureShareSample.cpp
--- a/Engine/Extras/VirtualProduction/TextureShare/Samples/D3D11/DC5_CopyToBackbuffer/Source/D3D11TextureShareSample.cpp
+++ b/Engine/Extras/VirtualProduction/TextureShare/Samples/D3D11/DC5_CopyToBackbuffer/Source/D3D11TextureShareSample.cpp
@@ -13,8 +13,8 @@ FD3D11TextureShareSample::FD3D11TextureShareSample(ID3D11Texture2D* InBackBuffer
 	if(TextureShareObject = ITextureShareObject::CreateInstance(TextureShareSample::ObjectDesc))
 	{
 		// Request resources for receive
-		TextureShareObject->GetData().ResourceRequests.Add(TextureShareSample::Receive::Texture1::Desc);
-		TextureShareObject->GetData().ResourceRequests.Add(TextureShareSample::Receive::Texture2::Desc);
+		TextureShareObject->GetData().ResourceRequests.Add(FTextureShareResourceRequest(TextureShareSample::Receive::Texture1::Desc, D3D11AppSetup::Backbuffer::Format));
+		TextureShareObject->GetData().ResourceRequests.Add(FTextureShareResourceRequest(TextureShareSample::Receive::Texture2::Desc, D3D11AppSetup::Backbuffer::Format));
 	}
 }
 
diff --git a/Engine/Extras/VirtualProduction/TextureShare/Samples/D3D11/DC5_CopyToBackbuffer/Source/D3D11TextureShareSampleSetup.h b/Engine/Extras/VirtualProduction/TextureShare/Samples/D3D11/DC5_CopyToBackbuffer/Source/D3D11TextureShareSampleSetup.h
--- a/Engine/Extras/VirtualProduction/TextureShare/Samples/D3D11/DC5_CopyToBackbuffer/Source/D3D11TextureShareSampleSetup.h
+++ b/Engine/Extras/VirtualProduction/TextureShare/Samples/D3D11/DC5_CopyToBackbuffer/Source/D3D11TextureShareSampleSetup.h
@@ -25,7 +25,7 @@ namespace TextureShareSample
 		namespace Texture2
 		{
 			// Request to read a resource #2 to a remote process
-			static FTextureShareViewportResourceDesc Desc(DisplayCluster::Viewport2, TextureShareStrings::SceneTextures::FinalColor, ETextureShareTextureOp::Read);
+			static FTextureShareViewportResourceDesc Desc(DisplayCluster::Viewport2, TextureShareDisplayClusterStrings::Output::Backbuffer, ETextureShareTextureOp::Read);
 		}
 	}
 };
diff --git a/Engine/Extras/VirtualProduction/TextureShare/Samples/D3D12/DC5_CopyToBackbuffer/Source/D3D12TextureShareSample.cpp b/Engine/Extras/VirtualProduction/TextureShare/Samples/D3D12/DC5_CopyToBackbuffer/Source/D3D12TextureShareSample.cpp
--- a/Engine/Extras/VirtualProduction/TextureShare/Samples/D3D12/DC5_CopyToBackbuffer/Source/D3D12TextureShareSample.cpp
+++ b/Engine/Extras/VirtualProduction/TextureShare/Samples/D3D12/DC5_CopyToBackbuffer/Source/D3D12TextureShareSample.cpp
@@ -13,8 +13,8 @@ FD3D12TextureShareSample::FD3D12TextureShareSample(ID3D12Resource* InBackBufferT
 	if(TextureShareObject = ITextureShareObject::CreateInstance(TextureShareSample::ObjectDesc))
 	{
 		// Request resources for receive
-		TextureShareObject->GetData().ResourceRequests.Add(TextureShareSample::Receive::Texture1::Desc);
-		TextureShareObject->GetData().ResourceRequests.Add(TextureShareSample::Receive::Texture2::Desc);
+		TextureShareObject->GetData().ResourceRequests.Add(FTextureShareResourceRequest(TextureShareSample::Receive::Texture1::Desc, D3D12AppSetup::Backbuffer::Format));
+		TextureShareObject->GetData().ResourceRequests.Add(FTextureShareResourceRequest(TextureShareSample::Receive::Texture2::Desc, D3D12AppSetup::Backbuffer::Format));
 	}
 }
 
diff --git a/Engine/Extras/VirtualProduction/TextureShare/Samples/D3D12/DC5_CopyToBackbuffer/Source/D3D12TextureShareSampleSetup.h b/Engine/Extras/VirtualProduction/TextureShare/Samples/D3D12/DC5_CopyToBackbuffer/Source/D3D12TextureShareSampleSetup.h
--- a/Engine/Extras/VirtualProduction/TextureShare/Samples/D3D12/DC5_CopyToBackbuffer/Source/D3D12TextureShareSampleSetup.h
+++ b/Engine/Extras/VirtualProduction/TextureShare/Samples/D3D12/DC5_CopyToBackbuffer/Source/D3D12TextureShareSampleSetup.h
@@ -25,7 +25,7 @@ namespace TextureShareSample
 		namespace Texture2
 		{
 			// Request to read a resource #2 to a remote process
-			static FTextureShareViewportResourceDesc Desc(DisplayCluster::Viewport2, TextureShareStrings::SceneTextures::FinalColor, ETextureShareTextureOp::Read);
+			static FTextureShareViewportResourceDesc Desc(DisplayCluster::Viewport2, TextureShareDisplayClusterStrings::Output::Backbuffer, ETextureShareTextureOp::Read);
 		}
 	}
 };
diff --git a/Engine/Extras/VirtualProduction/TextureShare/Samples/Helpers/Public/Containers/Resource/TextureShareResourceRequest.h b/Engine/Extras/VirtualProduction/TextureShare/Samples/Helpers/Public/Containers/Resource/TextureShareResourceRequest.h
--- a/Engine/Extras/VirtualProduction/TextureShare/Samples/Helpers/Public/Containers/Resource/TextureShareResourceRequest.h
+++ b/Engine/Extras/VirtualProduction/TextureShare/Samples/Helpers/Public/Containers/Resource/TextureShareResourceRequest.h
@@ -14,4 +14,12 @@ public:
 	FTextureShareResourceRequest(const FTextureShareCoreResourceDesc& InResourceDesc, const FTextureShareCustomResource& InCustomResource)
 		: FTextureShareCoreResourceRequest(InResourceDesc, InCustomResource.CustomSize, InCustomResource.CustomFormat)
 	{ }
+
+	FTextureShareResourceRequest(const FTextureShareCoreResourceDesc& InResourceDesc, const DXGI_FORMAT InCustomFormat)
+		: FTextureShareCoreResourceRequest(InResourceDesc, InCustomFormat)
+	{ }
+
+	FTextureShareResourceRequest(const FTextureShareCoreResourceDesc& InResourceDesc, const FIntPoint& InSize, const DXGI_FORMAT InFormat = DXGI_FORMAT_UNKNOWN)
+		: FTextureShareCoreResourceRequest(InResourceDesc, InSize, InFormat)
+	{ }
 };
diff --git a/Engine/Extras/VirtualProduction/TextureShare/Samples/Helpers/Public/Containers/Resource/TextureShareResourceSyncStep.h b/Engine/Extras/VirtualProduction/TextureShare/Samples/Helpers/Public/Containers/Resource/TextureShareResourceSyncStep.h
--- a/Engine/Extras/VirtualProduction/TextureShare/Samples/Helpers/Public/Containers/Resource/TextureShareResourceSyncStep.h
+++ b/Engine/Extras/VirtualProduction/TextureShare/Samples/Helpers/Public/Containers/Resource/TextureShareResourceSyncStep.h
@@ -22,8 +22,7 @@ namespace TextureShareResourceSyncStepHelper
 	static ETextureShareSyncStep GetDisplayClusterSyncStep(const wchar_t* InResourceName, const ETextureShareTextureOp InOperationType)
 	{
 		// Display cluster sync steps for resources:
-		if (IsEqual(InResourceName, TextureShareDisplayClusterStrings::Output::Viewport)
-			|| IsEqual(InResourceName, TextureShareDisplayClusterStrings::Output::Backbuffer)
+		if (IsEqual(InResourceName, TextureShareDisplayClusterStrings::Output::Backbuffer)
 			|| IsEqual(InResourceName, TextureShareDisplayClusterStrings::Output::BackbufferTemp))
 		{
 			return (InOperationType == ETextureShareTextureOp::Read) ? ETextureShareSyncStep::FrameProxyPostRenderBegin : ETextureShareSyncStep::FrameProxyPostRenderEnd;
diff --git a/Engine/Plugins/VirtualProduction/TextureShare/Source/TextureShare/Private/Game/ViewExtension/TextureShareSceneViewExtension.cpp b/Engine/Plugins/VirtualProduction/TextureShare/Source/TextureShare/Private/Game/ViewExtension/TextureShareSceneViewExtension.cpp
--- a/Engine/Plugins/VirtualProduction/TextureShare/Source/TextureShare/Private/Game/ViewExtension/TextureShareSceneViewExtension.cpp
+++ b/Engine/Plugins/VirtualProduction/TextureShare/Source/TextureShare/Private/Game/ViewExtension/TextureShareSceneViewExtension.cpp
@@ -68,76 +68,67 @@ namespace TextureShareSceneViewExtensionHelpers
 		OutViewFamily.GammaCorrection = InViewFamily.GammaCorrection;
 		OutViewFamily.SecondaryViewFraction = InViewFamily.SecondaryViewFraction;
 	}
+};
+using namespace TextureShareSceneViewExtensionHelpers;
 
-	BEGIN_SHADER_PARAMETER_STRUCT(FSendTextureParameters, )
-	RDG_TEXTURE_ACCESS(Texture, ERHIAccess::CopySrc)
-	END_SHADER_PARAMETER_STRUCT()
-
-	BEGIN_SHADER_PARAMETER_STRUCT(FReceiveTextureParameters, )
-	RDG_TEXTURE_ACCESS(Texture, ERHIAccess::CopyDest)
-	END_SHADER_PARAMETER_STRUCT()
-
-	static void SendRDGTexture(const TSharedPtr<ITextureShareObjectProxy, ESPMode::ThreadSafe> ObjectProxy, FRDGBuilder& GraphBuilder, FRDGTextureRef InTextureRef, const int32 InTextureGPUIndex, const FTextureShareCoreResourceDesc& InResourceDesc, const FIntRect& InViewRect)
+//////////////////////////////////////////////////////////////////////////////////////////////
+struct FTextureShareSceneView
+{
+	FTextureShareSceneView(const FSceneViewFamily& InViewFamily, const FSceneView& InSceneView, const FTextureShareSceneViewInfo& InViewInfo)
+		: ViewInfo(InViewInfo)
+		, SceneView(InSceneView)
 	{
-		FSendTextureParameters* PassParameters = GraphBuilder.AllocParameters<FSendTextureParameters>();
-		PassParameters->Texture = InTextureRef;
-
-		GraphBuilder.AddPass(
-			RDG_EVENT_NAME("TextureShare_SendRDGTexture_%s", *InResourceDesc.ResourceName),
-			PassParameters,
-			ERDGPassFlags::Copy | ERDGPassFlags::NeverCull,
-			[ObjectProxy, InTextureRef, InTextureGPUIndex, InResourceDesc, InViewRect](FRHICommandListImmediate& RHICmdList)
+#if WITH_MGPU
+		if (InViewFamily.bMultiGPUForkAndJoin)
 		{
-			ObjectProxy->ShareResource_RenderThread(RHICmdList, InResourceDesc, InTextureRef->GetRHI(), InTextureGPUIndex, &InViewRect);
-		});
+			GPUIndex = InSceneView.GPUMask.GetFirstIndex();
+		}
+#endif
+
+		UnconstrainedViewRect = SceneView.UnconstrainedViewRect;
+		UnscaledViewRect = SceneView.UnscaledViewRect;
 	}
 
+public:
+	int32 GPUIndex = -1;
+
+	FIntRect UnconstrainedViewRect;
+	FIntRect UnscaledViewRect;
+
+	FTextureShareSceneViewInfo ViewInfo;
+
+	const FSceneView& SceneView;
 };
-using namespace TextureShareSceneViewExtensionHelpers;
 
 //////////////////////////////////////////////////////////////////////////////////////////////
 // FTextureShareSceneViewExtension
 //////////////////////////////////////////////////////////////////////////////////////////////
-void FTextureShareSceneViewExtension::GetTextureShareCoreSceneViewData_RenderThread(const FSceneViewFamily& InViewFamily, const FSceneView& InSceneView, FTextureShareCoreSceneViewData& OutSceneViewData) const
+void FTextureShareSceneViewExtension::GetSceneViewData_RenderThread(const FTextureShareSceneView& InView)
 {
-	GetTextureShareCoreSceneView(InSceneView, OutSceneViewData.View);
-	GetTextureShareCoreSceneViewFamily(InViewFamily, OutSceneViewData.ViewFamily);
-}
+	const FSceneView& InSceneView = InView.SceneView;
+	if (InSceneView.Family)
+	{
+		const FSceneViewFamily& InViewFamily = *InSceneView.Family;
 
-void FTextureShareSceneViewExtension::GetSceneViewData_RenderThread(const FSceneView& InSceneView, const FTextureShareCoreViewDesc& InViewDesc)
-{
-	// Create new data container for viewport eye
-	FTextureShareCoreSceneViewData SceneViewData(InViewDesc);
+		// Create new data container for viewport eye
+		FTextureShareCoreSceneViewData SceneViewData(InView.ViewInfo.ViewDesc);
 
-	// Get view eye data
-	GetTextureShareCoreSceneViewData_RenderThread(*ViewFamily_RenderThread, InSceneView, SceneViewData);
+		// Get view eye data
+		GetTextureShareCoreSceneView(InSceneView, SceneViewData.View);
+		GetTextureShareCoreSceneViewFamily(InViewFamily, SceneViewData.ViewFamily);
 
-	// Save scene viewport eye data
-	ObjectProxy->GetCoreProxyData_RenderThread().SceneData.Add(SceneViewData);
+		// Save scene viewport eye data
+		ObjectProxy->GetCoreProxyData_RenderThread().SceneData.Add(SceneViewData);
+	}
 }
 
-void FTextureShareSceneViewExtension::ShareSceneViewColors_RenderThread(FRDGBuilder& GraphBuilder, const FSceneTextures& SceneTextures, const FSceneView& InSceneView, const FTextureShareCoreViewDesc& InViewDesc)
+void FTextureShareSceneViewExtension::ShareSceneViewColors_RenderThread(FRDGBuilder& GraphBuilder, const FSceneTextures& SceneTextures, const FTextureShareSceneView& InView)
 {
-	int32 SceneViewGPUIndex = -1;
-
-#if WITH_MGPU
-	if (ViewFamily_RenderThread->bMultiGPUForkAndJoin)
+	const auto AddShareTexturePass = [&](const TCHAR* InTextureName, const FRDGTextureRef& InTextureRef)
 	{
-		SceneViewGPUIndex = InSceneView.GPUMask.GetFirstIndex();
-	}
-#endif
-
-	const auto AddShareTexturePass = [&](const TCHAR* InTextureName, FRDGTextureRef InTextureRef)
-	{
-		if (HasBeenProduced(InTextureRef))
-		{
-			FTextureShareCoreResourceDesc InResourceDesc(InTextureName, InViewDesc, ETextureShareTextureOp::Read);
-			// Share only if the resource is requested from a remote process
-			if (const FTextureShareCoreResourceRequest* ExistResourceRequest = ObjectProxy->GetData_RenderThread().FindResourceRequest(InResourceDesc))
-			{
-				SendRDGTexture(ObjectProxy, GraphBuilder, InTextureRef, SceneViewGPUIndex, InResourceDesc, InSceneView.UnconstrainedViewRect);
-			}
-		}
+		// Send resource
+		ObjectProxy->ShareResource_RenderThread(GraphBuilder, FTextureShareCoreResourceDesc(InTextureName, InView.ViewInfo.ViewDesc, ETextureShareTextureOp::Read),
+			InTextureRef, InView.GPUIndex, &InView.UnconstrainedViewRect);
 	};
 
 	AddShareTexturePass(TextureShareStrings::SceneTextures::SceneColor, SceneTextures.Color.Resolve);
@@ -248,52 +239,48 @@ bool FTextureShareSceneViewExtension::IsActiveThisFrame_Internal(const FSceneVie
 	return (LinkedViewport == Context.Viewport);
 }
 
-void FTextureShareSceneViewExtension::SetupView(FSceneViewFamily& InViewFamily, FSceneView& InView)
+void FTextureShareSceneViewExtension::PreRenderViewFamily_RenderThread(FRDGBuilder& GraphBuilder, FSceneViewFamily& InViewFamily)
 {
 	FScopeLock Lock(&DataCS);
 
-}
+	// Reset values
+	Views.Empty();
 
-void FTextureShareSceneViewExtension::PreRenderViewFamily_RenderThread(FRHICommandListImmediate& RHICmdList, FSceneViewFamily& InViewFamily)
-{
-	FScopeLock Lock(&DataCS);
-
-	// be default scene view capturing is disabled
-	ViewFamily_RenderThread = nullptr;
-
-	if (bEnabled_RenderThread)
+	// Initialize views
+	for (const FSceneView* SceneViewIt : InViewFamily.Views)
 	{
-		ViewFamily_RenderThread = &InViewFamily;
+		if (SceneViewIt)
+		{
+			if (const FTextureShareSceneViewInfo* ViewInfo = ObjectProxy->GetData_RenderThread().Views.Find(SceneViewIt->StereoViewIndex, SceneViewIt->StereoPass))
+			{
+				Views.Add(FTextureShareSceneView(InViewFamily, *SceneViewIt, *ViewInfo));
+			}
+		}
 	}
 
-	// Handle view extension functor callbacks
-	if (PreRenderViewFamilyFunction)
-	{
-		PreRenderViewFamilyFunction(RHICmdList, *this);
-	}
+	// Add RDG pass
+	AddPass(GraphBuilder, RDG_EVENT_NAME("PreRenderViewFamily_RenderThread"), [this, &InViewFamily](FRHICommandListImmediate& RHICmdList)
+		{
+			PreRenderViewFamily_RenderThread(RHICmdList, InViewFamily);
+		});
 }
 
-void FTextureShareSceneViewExtension::OnResolvedSceneColor_RenderThread(FRDGBuilder& GraphBuilder, const FSceneTextures& SceneTextures)
+void FTextureShareSceneViewExtension::PostRenderViewFamily_RenderThread(FRDGBuilder& GraphBuilder, FSceneViewFamily& InViewFamily)
+{
+	AddPass(GraphBuilder, RDG_EVENT_NAME("PostRenderViewFamily_RenderThread"), [this, &InViewFamily](FRHICommandListImmediate& RHICmdList)
+		{
+			PostRenderViewFamily_RenderThread(RHICmdList, InViewFamily);
+		});
+}
+
+void FTextureShareSceneViewExtension::PreRenderViewFamily_RenderThread(FRHICommandListImmediate& RHICmdList, FSceneViewFamily& InViewFamily)
 {
 	FScopeLock Lock(&DataCS);
 
-	if (ViewFamily_RenderThread && bEnabled_RenderThread)
+	// Handle view extension functor callbacks
+	if (PreRenderViewFamilyFunction)
 	{
-		// Send all  textures from viewfamily views
-		for (const FSceneView* SceneViewIt : ViewFamily_RenderThread->Views)
-		{
-			if (SceneViewIt)
-			{
-				if (const FTextureShareSceneViewInfo* ViewInfo = ObjectProxy->GetData_RenderThread().Views.Find(SceneViewIt->StereoViewIndex, SceneViewIt->StereoPass))
-				{
-					// Always get scene view data
-					GetSceneViewData_RenderThread(*SceneViewIt, ViewInfo->ViewDesc);
-
-					// Send scene textures on request
-					ShareSceneViewColors_RenderThread(GraphBuilder, SceneTextures, *SceneViewIt, ViewInfo->ViewDesc);
-				}
-			}
-		}
+		PreRenderViewFamilyFunction(RHICmdList, *this);
 	}
 }
 
@@ -301,41 +288,28 @@ void FTextureShareSceneViewExtension::PostRenderViewFamily_RenderThread(FRHIComm
 {
 	FScopeLock Lock(&DataCS);
 
-	if (ViewFamily_RenderThread && bEnabled_RenderThread)
+	for (const FTextureShareSceneView& ViewIt : Views)
 	{
-		for (const FSceneView* SceneViewIt : InViewFamily.Views)
+		// Always get scene view data
+		GetSceneViewData_RenderThread(ViewIt);
+
+		// Share only if the resource is requested from a remote process
+		if (const FTextureShareCoreResourceRequest* ExistResourceRequest = ObjectProxy->GetData_RenderThread().FindResourceRequest(FTextureShareCoreResourceDesc(TextureShareStrings::SceneTextures::FinalColor, ViewIt.ViewInfo.ViewDesc, ETextureShareTextureOp::Undefined)))
 		{
-			if (SceneViewIt)
+			FTexture2DRHIRef RenderTargetTexture = InViewFamily.RenderTarget->GetRenderTargetTexture();
+			if (RenderTargetTexture.IsValid())
 			{
-				if (const FTextureShareSceneViewInfo* ViewInfo = ObjectProxy->GetData_RenderThread().Views.Find(SceneViewIt->StereoViewIndex, SceneViewIt->StereoPass))
+				// Send
+				const FTextureShareCoreResourceDesc SendResourceDesc(TextureShareStrings::SceneTextures::FinalColor, ViewIt.ViewInfo.ViewDesc, ETextureShareTextureOp::Read);
+				ObjectProxy->ShareResource_RenderThread(RHICmdList, SendResourceDesc, RenderTargetTexture, ViewIt.GPUIndex, &ViewIt.UnscaledViewRect);
+
+				if (bEnableObjectProxySync)
 				{
-					// Share only if the resource is requested from a remote process
-					if (const FTextureShareCoreResourceRequest* ExistResourceRequest = ObjectProxy->GetData_RenderThread().FindResourceRequest(FTextureShareCoreResourceDesc(TextureShareStrings::SceneTextures::FinalColor, ViewInfo->ViewDesc, ETextureShareTextureOp::Undefined)))
+					// Receive
+					const FTextureShareCoreResourceDesc ReceiveResourceDesc(TextureShareStrings::SceneTextures::FinalColor, ViewIt.ViewInfo.ViewDesc, ETextureShareTextureOp::Write, ETextureShareSyncStep::FrameSceneFinalColorEnd);
+					if (ObjectProxy->ShareResource_RenderThread(RHICmdList, ReceiveResourceDesc, RenderTargetTexture, ViewIt.GPUIndex, &ViewIt.UnscaledViewRect))
 					{
-						int32 SceneViewGPUIndex = -1;
-#if WITH_MGPU
-						if (ViewFamily_RenderThread->bMultiGPUForkAndJoin)
-						{
-							SceneViewGPUIndex = SceneViewIt->GPUMask.GetFirstIndex();
-						}
-#endif
-						FTexture2DRHIRef RenderTargetTexture = InViewFamily.RenderTarget->GetRenderTargetTexture();
-						if (RenderTargetTexture.IsValid())
-						{
-							// Send
-							const FTextureShareCoreResourceDesc SendResourceDesc(TextureShareStrings::SceneTextures::FinalColor, ViewInfo->ViewDesc, ETextureShareTextureOp::Read);
-							ObjectProxy->ShareResource_RenderThread(RHICmdList, SendResourceDesc, RenderTargetTexture, SceneViewGPUIndex, &SceneViewIt->UnscaledViewRect);
-
-							if (bEnableObjectProxySync)
-							{
-								// Receive
-								const FTextureShareCoreResourceDesc ReceiveResourceDesc(TextureShareStrings::SceneTextures::FinalColor, ViewInfo->ViewDesc, ETextureShareTextureOp::Write, ETextureShareSyncStep::FrameSceneFinalColorEnd);
-								if (ObjectProxy->ShareResource_RenderThread(RHICmdList, ReceiveResourceDesc, RenderTargetTexture, SceneViewGPUIndex, &SceneViewIt->UnscaledViewRect))
-								{
-									ObjectProxy->FrameSync_RenderThread(RHICmdList, ETextureShareSyncStep::FrameSceneFinalColorEnd);
-								}
-							}
-						}
+						ObjectProxy->FrameSync_RenderThread(RHICmdList, ETextureShareSyncStep::FrameSceneFinalColorEnd);
 					}
 				}
 			}
@@ -347,9 +321,17 @@ void FTextureShareSceneViewExtension::PostRenderViewFamily_RenderThread(FRHIComm
 	{
 		PostRenderViewFamilyFunction(RHICmdList, *this);
 	}
+}
 
-	// Clear stored value
-	ViewFamily_RenderThread = nullptr;
+void FTextureShareSceneViewExtension::OnResolvedSceneColor_RenderThread(FRDGBuilder& GraphBuilder, const FSceneTextures& SceneTextures)
+{
+	FScopeLock Lock(&DataCS);
+
+	// Send scene textures on request
+	for (const FTextureShareSceneView& ViewIt : Views)
+	{
+		ShareSceneViewColors_RenderThread(GraphBuilder, SceneTextures, ViewIt);
+	}
 }
 
 void FTextureShareSceneViewExtension::OnBackBufferReadyToPresent_RenderThread(SWindow& SlateWindow, const FTexture2DRHIRef& InBackbuffer)
diff --git a/Engine/Plugins/VirtualProduction/TextureShare/Source/TextureShare/Private/Game/ViewExtension/TextureShareSceneViewExtension.h b/Engine/Plugins/VirtualProduction/TextureShare/Source/TextureShare/Private/Game/ViewExtension/TextureShareSceneViewExtension.h
--- a/Engine/Plugins/VirtualProduction/TextureShare/Source/TextureShare/Private/Game/ViewExtension/TextureShareSceneViewExtension.h
+++ b/Engine/Plugins/VirtualProduction/TextureShare/Source/TextureShare/Private/Game/ViewExtension/TextureShareSceneViewExtension.h
@@ -6,6 +6,7 @@
 
 class ITextureShareObjectProxy;
 struct FTextureShareCoreViewDesc;
+struct FTextureShareSceneView;
 
 /**
  * Convenience type definition of a function that gives an opinion of whether the scene view extension should be active in the given context for the current frame.
@@ -30,23 +31,22 @@ public:
 	//~ Begin ISceneViewExtension interface
 	virtual int32 GetPriority() const override { return -1; }
 
-	virtual void SetupViewFamily(FSceneViewFamily& InViewFamily) override {}
-	virtual void BeginRenderViewFamily(FSceneViewFamily& InViewFamily) override {};
+	virtual void SetupViewFamily(FSceneViewFamily& InViewFamily) override { }
+	virtual void BeginRenderViewFamily(FSceneViewFamily& InViewFamily) override { };
+	virtual void SetupView(FSceneViewFamily& InViewFamily, FSceneView& InView) override { };
 
-	virtual void SetupView(FSceneViewFamily& InViewFamily, FSceneView& InView) override;
-	virtual void PreRenderViewFamily_RenderThread(FRHICommandListImmediate& RHICmdList, FSceneViewFamily& InViewFamily) override;
-	virtual void PostRenderViewFamily_RenderThread(FRHICommandListImmediate& RHICmdList, FSceneViewFamily& InViewFamily) override;
+	virtual void PreRenderViewFamily_RenderThread(FRDGBuilder& GraphBuilder, FSceneViewFamily& InViewFamily) override;
+	virtual void PostRenderViewFamily_RenderThread(FRDGBuilder& GraphBuilder, FSceneViewFamily& InViewFamily) override;
 
 	virtual void OnResolvedSceneColor_RenderThread(FRDGBuilder& GraphBuilder, const FSceneTextures& SceneTextures);
 	virtual void OnBackBufferReadyToPresent_RenderThread(SWindow&, const FTexture2DRHIRef&);
 
 	virtual bool IsActiveThisFrame_Internal(const FSceneViewExtensionContext& Context) const override;
-	//~End ISceneVIewExtension interface
+	//~End ISceneViewExtension interface
 
 public:
-	virtual void GetSceneViewData_RenderThread(const FSceneView& InSceneView, const FTextureShareCoreViewDesc& InViewDesc);
-	virtual void GetTextureShareCoreSceneViewData_RenderThread(const FSceneViewFamily& InViewFamily, const FSceneView& InSceneView, struct FTextureShareCoreSceneViewData& OutData) const;
-	virtual void ShareSceneViewColors_RenderThread(FRDGBuilder& GraphBuilder, const FSceneTextures& SceneTextures, const FSceneView& InSceneView, const FTextureShareCoreViewDesc& InViewDesc);
+	virtual void GetSceneViewData_RenderThread(const FTextureShareSceneView& InView);
+	virtual void ShareSceneViewColors_RenderThread(FRDGBuilder& GraphBuilder, const FSceneTextures& SceneTextures, const FTextureShareSceneView& InView);
 
 public:
 	// Returns true if the given object is of the same type.
@@ -59,13 +59,6 @@ public:
 	{
 		return ObjectProxy;
 	}
-	
-	FSceneViewFamily* GetViewFamily_RenderThread() const
-	{
-		check(IsInRenderingThread());
-
-		return ViewFamily_RenderThread;
-	}
 
 	void Initialize(const TSharedPtr<FTextureShareSceneViewExtension, ESPMode::ThreadSafe> InViewExtension);
 
@@ -80,6 +73,10 @@ public:
 	FViewport* GetLinkedViewport() const;
 	void SetLinkedViewport(FViewport* InLinkedViewport);
 
+private:
+	void PreRenderViewFamily_RenderThread(FRHICommandListImmediate& RHICmdList, FSceneViewFamily& InViewFamily);
+	void PostRenderViewFamily_RenderThread(FRHICommandListImmediate& RHICmdList, FSceneViewFamily& InViewFamily);
+
 protected:
 	mutable FCriticalSection DataCS;
 
@@ -89,7 +86,8 @@ protected:
 	bool bEnableObjectProxySync = false;
 
 	bool bEnabled = true;
-	bool bEnabled_RenderThread = true;
+
+	TArray<FTextureShareSceneView> Views;
 
 	TFunctionTextureShareViewExtension PreRenderViewFamilyFunction;
 	TFunctionTextureShareViewExtension PostRenderViewFamilyFunction;
@@ -98,7 +96,4 @@ protected:
 	
 	/** TextureShare proxy owner*/
 	TSharedRef<ITextureShareObjectProxy, ESPMode::ThreadSafe> ObjectProxy;
-
-	/** Runtime value */
-	FSceneViewFamily* ViewFamily_RenderThread = nullptr;
 };
diff --git a/Engine/Plugins/VirtualProduction/TextureShare/Source/TextureShare/Private/Game/WorldSubsystem/TextureShareWorldSubsystemObjectProxy.cpp b/Engine/Plugins/VirtualProduction/TextureShare/Source/TextureShare/Private/Game/WorldSubsystem/TextureShareWorldSubsystemObjectProxy.cpp
--- a/Engine/Plugins/VirtualProduction/TextureShare/Source/TextureShare/Private/Game/WorldSubsystem/TextureShareWorldSubsystemObjectProxy.cpp
+++ b/Engine/Plugins/VirtualProduction/TextureShare/Source/TextureShare/Private/Game/WorldSubsystem/TextureShareWorldSubsystemObjectProxy.cpp
@@ -75,11 +75,9 @@ FTextureShareWorldSubsystemObjectProxy::FTextureShareWorldSubsystemObjectProxy(c
 
 bool FTextureShareWorldSubsystemObjectProxy::Update(const UTextureShareObject* InTextureShareObject)
 {
-	TSharedPtr<FProxyResourcesData> ProxyData = MakeShared<FProxyResourcesData>(InTextureShareObject);
-
 	// Updated sceneproxy on renderthread
 	ENQUEUE_RENDER_COMMAND(TextureShare_UpdateObjectProxy)(
-		[ObjectProxy = ObjectProxy, ProxyData](FRHICommandListImmediate& RHICmdList)
+		[ObjectProxy = ObjectProxy, ProxyData = MakeShared<FProxyResourcesData>(InTextureShareObject)](FRHICommandListImmediate& RHICmdList)
 	{
 		FTextureShareWorldSubsystemObjectProxy ProxyAPI(ObjectProxy);
 		ProxyAPI.Update_RenderThread(RHICmdList, ProxyData);
@@ -107,46 +105,28 @@ void FTextureShareWorldSubsystemObjectProxy::Update_RenderThread(FRHICommandList
 				TSharedPtr<ITextureShareObjectProxy, ESPMode::ThreadSafe> ObjectProxy = InViewExtension.GetObjectProxy();
 				if (ObjectProxy.IsValid())
 				{
-					if (InViewExtension.GetViewFamily_RenderThread())
-					{
-						ObjectProxy->BeginFrameSync_RenderThread(RHICmdList);
-					}
-					else
-					{
-						// TextureShare not used for this viewfamily
-						if (ObjectProxy->BeginFrameSync_RenderThread(RHICmdList)
-							&& ObjectProxy->IsFrameSyncActive_RenderThread())
-						{
-							FTextureShareWorldSubsystemObjectProxy ProxyAPI(ObjectProxy);
-							ProxyAPI.UpdateResources_RenderThread(RHICmdList, ProxyData);
-
-							ObjectProxy->EndFrameSync_RenderThread(RHICmdList);
-						}
-					}
+					ObjectProxy->BeginFrameSync_RenderThread(RHICmdList);
 				}
 			};
 
 			TFunctionTextureShareViewExtension PostRenderViewFamilyFunction = [ProxyData, bBackbufferShared](FRHICommandListImmediate& RHICmdList, FTextureShareSceneViewExtension& InViewExtension)
 			{
-				if (InViewExtension.GetViewFamily_RenderThread())
+				// Share user resources
+				TSharedPtr<ITextureShareObjectProxy, ESPMode::ThreadSafe> ObjectProxy = InViewExtension.GetObjectProxy();
+				if (ObjectProxy.IsValid()
+					&& ObjectProxy->IsFrameSyncActive_RenderThread())
 				{
-					// Share user resources
-					TSharedPtr<ITextureShareObjectProxy, ESPMode::ThreadSafe> ObjectProxy = InViewExtension.GetObjectProxy();
-					if (ObjectProxy.IsValid()
-						&& ObjectProxy->IsFrameSyncActive_RenderThread())
-					{
-						FTextureShareWorldSubsystemObjectProxy ProxyAPI(ObjectProxy);
+					FTextureShareWorldSubsystemObjectProxy ProxyAPI(ObjectProxy);
 
-						ProxyAPI.UpdateResources_RenderThread(RHICmdList, ProxyData);
+					ProxyAPI.UpdateResources_RenderThread(RHICmdList, ProxyData);
 
-						if (bBackbufferShared)
-						{
-							ObjectProxy->FrameSync_RenderThread(RHICmdList, ETextureShareSyncStep::FrameProxyPreRenderEnd);
-						}
-						else
-						{
-							ObjectProxy->EndFrameSync_RenderThread(RHICmdList);
-						}
+					if (bBackbufferShared)
+					{
+						ObjectProxy->FrameSync_RenderThread(RHICmdList, ETextureShareSyncStep::FrameProxyPreRenderEnd);
+					}
+					else
+					{
+						ObjectProxy->EndFrameSync_RenderThread(RHICmdList);
 					}
 				}
 			};
diff --git a/Engine/Plugins/VirtualProduction/TextureShare/Source/TextureShare/Private/ITextureShareObjectProxy.h b/Engine/Plugins/VirtualProduction/TextureShare/Source/TextureShare/Private/ITextureShareObjectProxy.h
--- a/Engine/Plugins/VirtualProduction/TextureShare/Source/TextureShare/Private/ITextureShareObjectProxy.h
+++ b/Engine/Plugins/VirtualProduction/TextureShare/Source/TextureShare/Private/ITextureShareObjectProxy.h
@@ -2,14 +2,11 @@
 
 #pragma once
 #include "CoreMinimal.h"
+#include "RendererInterface.h"
 
 #include "Containers/TextureShareContainers.h"
 #include "Containers/TextureShareCoreContainers.h"
 
-#include "RHI.h"
-#include "RHICommandList.h"
-#include "RHIResources.h"
-
 class FTextureShareSceneViewExtension;
 
 /**
@@ -96,11 +93,13 @@ public:
 	virtual const TSharedPtr<FTextureShareSceneViewExtension, ESPMode::ThreadSafe>& GetViewExtension_RenderThread() const = 0;
 
 public:
+	/////////////////////////// Resource ///////////////////////////
+
 	/**
 	 * Share texture resource with ResourceRequest
 	 *
 	 * @param RHICmdList        - RHI cmd list
-	 * @param InResourceRequest - resource information
+	 * @param InResourceDesc    - resource information
 	 * @param InTexture         - Texture RHI resource
 	 * @param InTextureGPUIndex - Texture GPU index
 	 * @param InTextureRect     - (optional) Share only region from InTexture
@@ -108,4 +107,17 @@ public:
 	 * @return True if the success
 	 */
 	virtual bool ShareResource_RenderThread(FRHICommandListImmediate& RHICmdList, const FTextureShareCoreResourceDesc& InResourceDesc, FRHITexture* InTexture, const int32 InTextureGPUIndex, const FIntRect* InTextureRect = nullptr) const = 0;
+
+	/**
+	 * Share texture resource with ResourceRequest
+	 *
+	 * @param GraphBuilder      - RDG builder
+	 * @param InResourceDesc    - resource information
+	 * @param InTextureRef      - RDG Texture ref
+	 * @param InTextureGPUIndex - Texture GPU index
+	 * @param InTextureRect     - (optional) Share only region from InTexture
+	 *
+	 * @return True if the success
+	 */
+	virtual bool ShareResource_RenderThread(FRDGBuilder& GraphBuilder, const FTextureShareCoreResourceDesc& InResourceDesc, const FRDGTextureRef& InTextureRef, const int32 InTextureGPUIndex, const FIntRect* InTextureRect = nullptr) const = 0;
 };
diff --git a/Engine/Plugins/VirtualProduction/TextureShare/Source/TextureShare/Private/Module/TextureShareAPI.cpp b/Engine/Plugins/VirtualProduction/TextureShare/Source/TextureShare/Private/Module/TextureShareAPI.cpp
--- a/Engine/Plugins/VirtualProduction/TextureShare/Source/TextureShare/Private/Module/TextureShareAPI.cpp
+++ b/Engine/Plugins/VirtualProduction/TextureShare/Source/TextureShare/Private/Module/TextureShareAPI.cpp
@@ -78,7 +78,7 @@ TSharedPtr<ITextureShareObject, ESPMode::ThreadSafe> FTextureShareAPI::GetOrCrea
 		CoreObject->SetDeviceType(GetTextureShareDeviceType());
 
 		// Create game thread object
-		TSharedPtr<FTextureShareObject, ESPMode::ThreadSafe> NewObject = MakeShared<FTextureShareObject, ESPMode::ThreadSafe>(CoreObject.ToSharedRef());
+		const TSharedPtr<FTextureShareObject, ESPMode::ThreadSafe> NewObject(new FTextureShareObject(CoreObject.ToSharedRef()));
 		if (NewObject.IsValid())
 		{
 			// Register for game thread
diff --git a/Engine/Plugins/VirtualProduction/TextureShare/Source/TextureShare/Private/Object/TextureShareObject.cpp b/Engine/Plugins/VirtualProduction/TextureShare/Source/TextureShare/Private/Object/TextureShareObject.cpp
--- a/Engine/Plugins/VirtualProduction/TextureShare/Source/TextureShare/Private/Object/TextureShareObject.cpp
+++ b/Engine/Plugins/VirtualProduction/TextureShare/Source/TextureShare/Private/Object/TextureShareObject.cpp
@@ -12,7 +12,8 @@
 //////////////////////////////////////////////////////////////////////////////////////////////
 FTextureShareObject::FTextureShareObject(const TSharedRef<ITextureShareCoreObject, ESPMode::ThreadSafe>& InCoreObject)
 	: CoreObject(InCoreObject)
-	, ObjectProxy(MakeShared<FTextureShareObjectProxy, ESPMode::ThreadSafe>(CoreObject))
+	, ObjectProxy(new FTextureShareObjectProxy(CoreObject))
+	, TextureShareData(MakeShared<FTextureShareData, ESPMode::ThreadSafe>())
 { }
 
 FTextureShareObject::~FTextureShareObject()
@@ -121,18 +122,17 @@ bool FTextureShareObject::EndFrameSync(FViewport* InViewport)
 	bFrameSyncActive = false;
 
 	// Store gamethread data for proxy
-	FTextureShareData& ProxyObjectData = GetData();
-	ProxyObjectData.ObjectData = CoreObject->GetData();
-	ProxyObjectData.ReceivedObjectsData.Append(CoreObject->GetReceivedData());
+	TextureShareData->ObjectData = CoreObject->GetData();
+	TextureShareData->ReceivedObjectsData.Append(CoreObject->GetReceivedData());
 
 	FTextureShareObjectProxy::UpdateProxy_GameThread(*this);
 
+	// Reset ptr to data, now this data used in proxy
+	TextureShareData = MakeShared<FTextureShareData, ESPMode::ThreadSafe>();
+
 	// Game thread data now sent to the proxy.Game-thread data can now be cleared
 	const bool bResult = CoreObject->EndFrameSync();
 
-	// Reset ptr to data, now this data used in proxy
-	ObjectData.Reset();
-
 	CoreObject->UnlockThreadMutex(ETextureShareThreadMutex::RenderingThread);
 
 	return bResult;
@@ -172,12 +172,7 @@ const TArray<FTextureShareCoreObjectData>& FTextureShareObject::GetReceivedCoreO
 //////////////////////////////////////////////////////////////////////////////////////////////
 FTextureShareData& FTextureShareObject::GetData()
 {
-	if (!ObjectData.IsValid())
-	{
-		ObjectData = MakeShared<FTextureShareData>();
-	}
-
-	return *ObjectData;
+	return *TextureShareData;
 }
 
 TSharedPtr<FTextureShareSceneViewExtension, ESPMode::ThreadSafe> FTextureShareObject::GetViewExtension() const
diff --git a/Engine/Plugins/VirtualProduction/TextureShare/Source/TextureShare/Private/Object/TextureShareObject.h b/Engine/Plugins/VirtualProduction/TextureShare/Source/TextureShare/Private/Object/TextureShareObject.h
--- a/Engine/Plugins/VirtualProduction/TextureShare/Source/TextureShare/Private/Object/TextureShareObject.h
+++ b/Engine/Plugins/VirtualProduction/TextureShare/Source/TextureShare/Private/Object/TextureShareObject.h
@@ -5,6 +5,7 @@
 #include "ITextureShareCoreObject.h"
 
 #include "Game/ViewExtension/TextureShareSceneViewExtension.h"
+#include "Templates/SharedPointer.h"
 
 class FViewport;
 class FTextureShareObjectProxy;
@@ -14,6 +15,7 @@ class FTextureShareObjectProxy;
  */
 class FTextureShareObject
 	: public ITextureShareObject
+	, public TSharedFromThis<FTextureShareObject, ESPMode::ThreadSafe>
 {
 public:
 	FTextureShareObject(const TSharedRef<ITextureShareCoreObject, ESPMode::ThreadSafe>& InCoreObject);
@@ -74,7 +76,7 @@ protected:
 	const TSharedRef<FTextureShareObjectProxy, ESPMode::ThreadSafe> ObjectProxy;
 
 	// Object data from game thread
-	TSharedPtr<FTextureShareData, ESPMode::ThreadSafe> ObjectData;
+	TSharedRef<FTextureShareData, ESPMode::ThreadSafe> TextureShareData;
 
 	bool bFrameSyncActive = false;
 	bool bSessionActive = false;
diff --git a/Engine/Plugins/VirtualProduction/TextureShare/Source/TextureShare/Private/Object/TextureShareObjectProxy.cpp b/Engine/Plugins/VirtualProduction/TextureShare/Source/TextureShare/Private/Object/TextureShareObjectProxy.cpp
--- a/Engine/Plugins/VirtualProduction/TextureShare/Source/TextureShare/Private/Object/TextureShareObjectProxy.cpp
+++ b/Engine/Plugins/VirtualProduction/TextureShare/Source/TextureShare/Private/Object/TextureShareObjectProxy.cpp
@@ -11,11 +11,14 @@
 
 #include "ITextureShareCore.h"
 
+#include "RenderGraphUtils.h"
+
 //////////////////////////////////////////////////////////////////////////////////////////////
 // FTextureShareObjectProxy
 //////////////////////////////////////////////////////////////////////////////////////////////
 FTextureShareObjectProxy::FTextureShareObjectProxy(const TSharedRef<ITextureShareCoreObject, ESPMode::ThreadSafe>& InCoreObject)
 	: CoreObject(InCoreObject)
+	, TextureShareData(MakeShared<FTextureShareData, ESPMode::ThreadSafe>())
 { }
 
 FTextureShareObjectProxy::~FTextureShareObjectProxy()
@@ -64,25 +67,24 @@ bool FTextureShareObjectProxy::BeginFrameSync_RenderThread(FRHICommandListImmedi
 	{
 		if (CoreObject->BeginFrameSync_RenderThread())
 		{
-			if (ObjectData.IsValid())
-			{
-				// update frame markers from game thread data
-				FTextureShareCoreProxyData& ProxyDataRef = CoreObject->GetProxyData_RenderThread();
+			// update frame markers from game thread data
+			FTextureShareCoreProxyData& CoreProxyDataRef = CoreObject->GetProxyData_RenderThread();
 				
-				// Copy frame marker from game thread
-				ProxyDataRef.FrameMarker = ObjectData->ObjectData.FrameMarker;
-
-				// Copy the frame markers from the objects saved at the end of the game stream.
-				ProxyDataRef.RemoteFrameMarkers.Empty();
-				for (const FTextureShareCoreObjectData& ObjectDataIt : ObjectData->ReceivedObjectsData)
-				{
-					ProxyDataRef.RemoteFrameMarkers.Add(FTextureShareCoreObjectFrameMarker(ObjectDataIt.Desc, ObjectDataIt.Data.FrameMarker));
-				}
+			// Copy frame marker from game thread
+			CoreProxyDataRef.FrameMarker = TextureShareData->ObjectData.FrameMarker;
 
-				bFrameProxySyncActive = true;
+			// Copy the frame markers from the objects saved at the end of the game stream.
+			CoreProxyDataRef.RemoteFrameMarkers.Empty();
 
-				return true;
+			// Update remote frame markers
+			for (const FTextureShareCoreObjectData& ObjectDataIt : TextureShareData->ReceivedObjectsData)
+			{
+				CoreProxyDataRef.RemoteFrameMarkers.Add(FTextureShareCoreObjectFrameMarker(ObjectDataIt.Desc, ObjectDataIt.Data.FrameMarker));
 			}
+
+			bFrameProxySyncActive = true;
+
+			return true;
 		}
 	}
 
@@ -147,13 +149,7 @@ const FTextureShareData& FTextureShareObjectProxy::GetData_RenderThread() const
 {
 	check(IsInRenderingThread());
 
-	if (!ObjectData.IsValid())
-	{
-		static FTextureShareData EmptyData;
-		return EmptyData;
-	}
-
-	return *ObjectData;
+	return *TextureShareData;
 }
 
 const TSharedPtr<FTextureShareSceneViewExtension, ESPMode::ThreadSafe>& FTextureShareObjectProxy::GetViewExtension_RenderThread() const
@@ -184,6 +180,67 @@ const TArray<FTextureShareCoreObjectProxyData>& FTextureShareObjectProxy::GetRec
 	return CoreObject->GetReceivedProxyData_RenderThread();
 }
 
+//////////////////////////////////////////////////////////////////////////////////////////////
+BEGIN_SHADER_PARAMETER_STRUCT(FSendTextureParameters, )
+RDG_TEXTURE_ACCESS(Texture, ERHIAccess::CopySrc)
+END_SHADER_PARAMETER_STRUCT()
+
+BEGIN_SHADER_PARAMETER_STRUCT(FReceiveTextureParameters, )
+RDG_TEXTURE_ACCESS(Texture, ERHIAccess::CopyDest)
+END_SHADER_PARAMETER_STRUCT()
+
+bool FTextureShareObjectProxy::ShareResource_RenderThread(FRDGBuilder& GraphBuilder, const FTextureShareCoreResourceDesc& InResourceDesc, const FRDGTextureRef& InTextureRef, const int32 InTextureGPUIndex, const FIntRect* InTextureRect) const
+{
+	if (HasBeenProduced(InTextureRef))
+	{
+		FIntRect InViewRect = (InTextureRect) ? *InTextureRect : FIntRect();
+
+		switch (InResourceDesc.OperationType)
+		{
+		case ETextureShareTextureOp::Read:
+		{
+			FSendTextureParameters* PassParameters = GraphBuilder.AllocParameters<FSendTextureParameters>();
+			PassParameters->Texture = InTextureRef;
+
+			GraphBuilder.AddPass(
+				RDG_EVENT_NAME("TextureShare_SendRDGTexture_%s", *InResourceDesc.ResourceName),
+				PassParameters,
+				ERDGPassFlags::Copy | ERDGPassFlags::NeverCull,
+				[ObjectProxy = SharedThis(this), SyncStep = InResourceDesc.SyncStep, InResourceDesc, InTextureRef, InTextureGPUIndex, InViewRect](FRHICommandListImmediate& RHICmdList)
+				{
+					ObjectProxy->ShareResource_RenderThread(RHICmdList, InResourceDesc, InTextureRef->GetRHI(), InTextureGPUIndex, &InViewRect);
+				});
+
+			return true;
+		}
+		break;
+
+		case ETextureShareTextureOp::Write:
+		{
+			FReceiveTextureParameters* PassParameters = GraphBuilder.AllocParameters<FReceiveTextureParameters>();
+			PassParameters->Texture = InTextureRef;
+
+			GraphBuilder.AddPass(
+				RDG_EVENT_NAME("TextureShare_ReceiveRDGTexture_%s", *InResourceDesc.ResourceName),
+				PassParameters,
+				ERDGPassFlags::Copy | ERDGPassFlags::NeverCull,
+				[ObjectProxy = SharedThis(this), InTextureRef, InTextureGPUIndex, InResourceDesc, InViewRect](FRHICommandListImmediate& RHICmdList)
+				{
+					ObjectProxy->ShareResource_RenderThread(RHICmdList, InResourceDesc, InTextureRef->GetRHI(), InTextureGPUIndex, &InViewRect);
+				});
+
+			return true;
+		}
+		break;
+
+		default:
+			break;
+		}
+	}
+
+	return false;
+}
+
 //////////////////////////////////////////////////////////////////////////////////////////////
 bool FTextureShareObjectProxy::ShareResource_RenderThread(FRHICommandListImmediate& RHICmdList, const FTextureShareCoreResourceDesc& InResourceDesc, FRHITexture* InTexture, const int32 InTextureGPUIndex, const FIntRect* InTextureRect) const
 {
@@ -265,12 +322,12 @@ bool FTextureShareObjectProxy::EndSession_RenderThread()
 	return false;
 }
 
-void FTextureShareObjectProxy::HandleNewFrame_RenderThread(const TSharedPtr<FTextureShareData, ESPMode::ThreadSafe>& InObjectData, const TSharedPtr<FTextureShareSceneViewExtension, ESPMode::ThreadSafe>& InViewExtension)
+void FTextureShareObjectProxy::HandleNewFrame_RenderThread(const TSharedRef<FTextureShareData, ESPMode::ThreadSafe>& InTextureShareData, const TSharedPtr<FTextureShareSceneViewExtension, ESPMode::ThreadSafe>& InViewExtension)
 {
 	check(IsInRenderingThread());
 
 	// Assign new frame data sent from game thread
-	ObjectData = InObjectData;
+	TextureShareData = InTextureShareData;
 
 	ViewExtension = InViewExtension;
 
@@ -308,9 +365,9 @@ void FTextureShareObjectProxy::UpdateProxy_GameThread(const FTextureShareObject&
 {
 	ENQUEUE_RENDER_COMMAND(TextureShare_UpdateObjectProxy)(
 		[ObjectProxyRef = In.GetObjectProxyRef()
-		, NewData = In.ObjectData
+		, NewTextureShareData = In.TextureShareData
 		, SceneViewExtension = In.ViewExtension](FRHICommandListImmediate& RHICmdList)
 	{
-		ObjectProxyRef->HandleNewFrame_RenderThread(NewData, SceneViewExtension);
+		ObjectProxyRef->HandleNewFrame_RenderThread(NewTextureShareData, SceneViewExtension);
 	});
 }
diff --git a/Engine/Plugins/VirtualProduction/TextureShare/Source/TextureShare/Private/Object/TextureShareObjectProxy.h b/Engine/Plugins/VirtualProduction/TextureShare/Source/TextureShare/Private/Object/TextureShareObjectProxy.h
--- a/Engine/Plugins/VirtualProduction/TextureShare/Source/TextureShare/Private/Object/TextureShareObjectProxy.h
+++ b/Engine/Plugins/VirtualProduction/TextureShare/Source/TextureShare/Private/Object/TextureShareObjectProxy.h
@@ -4,6 +4,7 @@
 #include "ITextureShareObjectProxy.h"
 #include "ITextureShareCoreObject.h"
 #include "Containers/TextureShareContainers.h"
+#include "Templates/SharedPointer.h"
 
 class FTextureShareResourcesProxy;
 class FTextureShareSceneViewExtension;
@@ -13,6 +14,7 @@ class FTextureShareSceneViewExtension;
  */
 class FTextureShareObjectProxy
 	: public ITextureShareObjectProxy
+	, public TSharedFromThis<FTextureShareObjectProxy, ESPMode::ThreadSafe>
 {
 	friend class FTextureShareObject;
 
@@ -41,12 +43,13 @@ public:
 	virtual const TSharedPtr<FTextureShareSceneViewExtension, ESPMode::ThreadSafe>& GetViewExtension_RenderThread() const override;
 
 	virtual bool ShareResource_RenderThread(FRHICommandListImmediate& RHICmdList, const FTextureShareCoreResourceDesc& InResourceDesc, FRHITexture* InTexture, const int32 InTextureGPUIndex, const FIntRect* InTextureRect = nullptr) const override;
+	virtual bool ShareResource_RenderThread(FRDGBuilder& GraphBuilder, const FTextureShareCoreResourceDesc& InResourceDesc, const FRDGTextureRef& InTextureRef, const int32 InTextureGPUIndex, const FIntRect* InTextureRect = nullptr) const override;
 	//~ITextureShareObjectProxy
 
 protected:
 	bool BeginSession_RenderThread();
 	bool EndSession_RenderThread();
-	void HandleNewFrame_RenderThread(const TSharedPtr<FTextureShareData, ESPMode::ThreadSafe>& InData, const TSharedPtr<FTextureShareSceneViewExtension, ESPMode::ThreadSafe>& InViewExtension);
+	void HandleNewFrame_RenderThread(const TSharedRef<FTextureShareData, ESPMode::ThreadSafe>& InTextureShareData, const TSharedPtr<FTextureShareSceneViewExtension, ESPMode::ThreadSafe>& InViewExtension);
 
 	static void BeginSession_GameThread(const FTextureShareObject& In);
 	static void EndSession_GameThread(const FTextureShareObject& In);
@@ -57,7 +60,7 @@ private:
 	const TSharedRef<ITextureShareCoreObject, ESPMode::ThreadSafe> CoreObject;
 
 	// Object data from game thread
-	TSharedPtr<FTextureShareData, ESPMode::ThreadSafe> ObjectData;
+	TSharedRef<FTextureShareData, ESPMode::ThreadSafe> TextureShareData;
 
 	// Frame proxy sync valid
 	mutable bool bFrameProxySyncActive = false;
diff --git a/Engine/Plugins/VirtualProduction/TextureShare/Source/TextureShareCore/Private/Containers/TextureShareCoreContainers_ResourceRequest.h b/Engine/Plugins/VirtualProduction/TextureShare/Source/TextureShareCore/Private/Containers/TextureShareCoreContainers_ResourceRequest.h
--- a/Engine/Plugins/VirtualProduction/TextureShare/Source/TextureShareCore/Private/Containers/TextureShareCoreContainers_ResourceRequest.h
+++ b/Engine/Plugins/VirtualProduction/TextureShare/Source/TextureShareCore/Private/Containers/TextureShareCoreContainers_ResourceRequest.h
@@ -53,6 +53,11 @@ public:
 		, Size(InSize)
 	{ }
 
+	FTextureShareCoreResourceRequest(const FTextureShareCoreResourceDesc& InResourceDesc, const DXGI_FORMAT InFormat)
+		: ResourceDesc(InResourceDesc)
+		, Format(InFormat)
+	{ }
+
 	void SetPixelFormat(uint32 InPixelFormat)
 	{
 		PixelFormat = (EPixelFormat)InPixelFormat;
diff --git a/Engine/Plugins/VirtualProduction/TextureShare/Source/TextureShareDisplayCluster/Private/Misc/TextureShareDisplayClusterStrings.h b/Engine/Plugins/VirtualProduction/TextureShare/Source/TextureShareDisplayCluster/Private/Misc/TextureShareDisplayClusterStrings.h
--- a/Engine/Plugins/VirtualProduction/TextureShare/Source/TextureShareDisplayCluster/Private/Misc/TextureShareDisplayClusterStrings.h
+++ b/Engine/Plugins/VirtualProduction/TextureShare/Source/TextureShareDisplayCluster/Private/Misc/TextureShareDisplayClusterStrings.h
@@ -25,13 +25,8 @@ namespace TextureShareDisplayClusterStrings
 		static constexpr auto Warped = TEXT("ViewportWarpBlend");
 	}
 
-	// final frame resource (after warp, output remaps)
 	namespace Output
 	{
-		// name mapped as viewport resource
-		// Ex.: "vp_0.BackbufferRect"
-		static constexpr auto Viewport = TEXT("FrameBackbufferRect");
-
 		// access to nDisplay frame backbuffer
 		static constexpr auto Backbuffer     = TEXT("FrameBackbuffer");
 		static constexpr auto BackbufferTemp = TEXT("FrameBackbufferTemp");
diff --git a/Engine/Plugins/VirtualProduction/TextureShare/Source/TextureShareDisplayCluster/Private/PostProcess/TextureSharePostprocess.cpp b/Engine/Plugins/VirtualProduction/TextureShare/Source/TextureShareDisplayCluster/Private/PostProcess/TextureSharePostprocess.cpp
--- a/Engine/Plugins/VirtualProduction/TextureShare/Source/TextureShareDisplayCluster/Private/PostProcess/TextureSharePostprocess.cpp
+++ b/Engine/Plugins/VirtualProduction/TextureShare/Source/TextureShareDisplayCluster/Private/PostProcess/TextureSharePostprocess.cpp
@@ -204,8 +204,11 @@ void FTextureSharePostprocess::HandleEndUpdateFrameResources_RenderThread(FRHICo
 
 	if (IsActive() && ObjectProxy->IsFrameSyncActive_RenderThread())
 	{
-		// Share RTT with remote process
-		ShareViewport_RenderThread(RHICmdList, InViewportManagerProxy, ETextureShareSyncStep::FrameProxyPostRenderEnd, EDisplayClusterViewportResourceType::OutputFrameTargetableResource,  TextureShareDisplayClusterStrings::Output::Viewport);
+		// Share per-viewport
+		ShareViewport_RenderThread(RHICmdList, InViewportManagerProxy, ETextureShareSyncStep::FrameProxyPostRenderEnd, EDisplayClusterViewportResourceType::OutputFrameTargetableResource, TextureShareDisplayClusterStrings::Output::Backbuffer);
+		ShareViewport_RenderThread(RHICmdList, InViewportManagerProxy, ETextureShareSyncStep::FrameProxyPostRenderEnd, EDisplayClusterViewportResourceType::AdditionalFrameTargetableResource, TextureShareDisplayClusterStrings::Output::BackbufferTemp);
+
+		//Share whole backbuffer
 		ShareFrame_RenderThread(RHICmdList, InViewportManagerProxy, ETextureShareSyncStep::FrameProxyPostRenderEnd, EDisplayClusterViewportResourceType::OutputFrameTargetableResource,     TextureShareDisplayClusterStrings::Output::Backbuffer);
 		ShareFrame_RenderThread(RHICmdList, InViewportManagerProxy, ETextureShareSyncStep::FrameProxyPostRenderEnd, EDisplayClusterViewportResourceType::AdditionalFrameTargetableResource, TextureShareDisplayClusterStrings::Output::BackbufferTemp);
 
