TextureShare: fixed memory leak

diff --git a/Engine/Plugins/VirtualProduction/TextureShare/Source/TextureShare/Private/Game/ViewExtension/TextureShareSceneViewExtension.cpp b/Engine/Plugins/VirtualProduction/TextureShare/Source/TextureShare/Private/Game/ViewExtension/TextureShareSceneViewExtension.cpp
--- a/Engine/Plugins/VirtualProduction/TextureShare/Source/TextureShare/Private/Game/ViewExtension/TextureShareSceneViewExtension.cpp
+++ b/Engine/Plugins/VirtualProduction/TextureShare/Source/TextureShare/Private/Game/ViewExtension/TextureShareSceneViewExtension.cpp
@@ -118,7 +118,15 @@ void FTextureShareSceneViewExtension::GetSceneViewData_RenderThread(const FTextu
 		GetTextureShareCoreSceneViewFamily(InViewFamily, SceneViewData.ViewFamily);
 
 		// Save scene viewport eye data
-		ObjectProxy->GetCoreProxyData_RenderThread().SceneData.Add(SceneViewData);
+		TArraySerializable<FTextureShareCoreSceneViewData>& DstSceneData = ObjectProxy->GetCoreProxyData_RenderThread().SceneData;
+		if (FTextureShareCoreSceneViewData* ExistValue = DstSceneData.FindByEqualsFunc(SceneViewData.ViewDesc))
+		{
+			*ExistValue = SceneViewData;
+		}
+		else
+		{
+			DstSceneData.Add(SceneViewData);
+		}
 	}
 }
 
diff --git a/Engine/Plugins/VirtualProduction/TextureShare/Source/TextureShareCore/Private/Serialize/TextureShareCoreSerialize.h b/Engine/Plugins/VirtualProduction/TextureShare/Source/TextureShareCore/Private/Serialize/TextureShareCoreSerialize.h
--- a/Engine/Plugins/VirtualProduction/TextureShare/Source/TextureShareCore/Private/Serialize/TextureShareCoreSerialize.h
+++ b/Engine/Plugins/VirtualProduction/TextureShare/Source/TextureShareCore/Private/Serialize/TextureShareCoreSerialize.h
@@ -168,6 +168,14 @@ public:
 		});
 	}
 
+	template<typename SubT>
+	T* FindByEqualsFunc(const SubT& InValue)
+	{
+		return FindByPredicate([InValue](const T& DataIt) {
+			return DataIt.EqualsFunc(InValue);
+		});
+	}
+
 	template<typename SubT>
 	bool GetValuesByEqualsFunc(const SubT& InValue, TArray<T>& OutValues) const
 	{
