Removing access to Private RHI code from TextureShare. Changed the code to use the new RHI interfaces.

diff --git a/Engine/Plugins/VirtualProduction/TextureShare/Source/TextureShare/Private/Resources/TextureShareResource_D3D11.cpp b/Engine/Plugins/VirtualProduction/TextureShare/Source/TextureShare/Private/Resources/TextureShareResource_D3D11.cpp
--- a/Engine/Plugins/VirtualProduction/TextureShare/Source/TextureShare/Private/Resources/TextureShareResource_D3D11.cpp
+++ b/Engine/Plugins/VirtualProduction/TextureShare/Source/TextureShare/Private/Resources/TextureShareResource_D3D11.cpp
@@ -8,15 +8,7 @@
 #include "ITextureShareCoreObject.h"
 #include "ITextureShareCoreD3D11ResourcesCache.h"
 
-#include "RHI.h"
-#include "RenderResource.h"
-
-#include "Windows/AllowWindowsPlatformTypes.h"
-#include "d3d11.h"
-#include "Windows/HideWindowsPlatformTypes.h"
-
-#include "D3D11RHIPrivate.h"
-#include "D3D11Util.h"
+#include "ID3D11DynamicRHI.h"
 
 //////////////////////////////////////////////////////////////////////////////////////////////
 namespace TextureShareResourceHelpers
@@ -35,22 +27,23 @@ using namespace TextureShareResourceHelpers;
 //////////////////////////////////////////////////////////////////////////////////////////////
 bool FTextureShareResource::D3D11RegisterResourceHandle(const FTextureShareCoreResourceRequest& InResourceRequest)
 {
-	if (GD3D11RHI && TextureRHI.IsValid())
+	if (IsRHID3D11() && TextureRHI.IsValid())
 	{
 		TSharedPtr<ITextureShareCoreD3D11ResourcesCache, ESPMode::ThreadSafe> D3D11ResourcesCache = GetD3D11ResourcesCache();
 		if (D3D11ResourcesCache.IsValid())
 		{
-			if (ID3D11Device* D3D11DevicePtr = static_cast<ID3D11Device*>(GD3D11RHI->RHIGetNativeDevice()))
+			ID3D11DynamicRHI* D3D11RHI = GetID3D11DynamicRHI();
+			ID3D11Device* D3D11DevicePtr = D3D11RHI->RHIGetDevice();
+			ID3D11Texture2D* D3D11ResourcePtr = (ID3D11Texture2D*)D3D11RHI->RHIGetResource(TextureRHI);
+
+			if (D3D11DevicePtr && D3D11ResourcePtr)
 			{
-				if (ID3D11Texture2D* D3D11ResourcePtr = (ID3D11Texture2D*)(TextureRHI->GetNativeResource()))
+				FTextureShareCoreResourceHandle ResourceHandle;
+				if (D3D11ResourcesCache->CreateSharedResource(CoreObject->GetObjectDesc(), D3D11DevicePtr, D3D11ResourcePtr, ResourceDesc, ResourceHandle))
 				{
-					FTextureShareCoreResourceHandle ResourceHandle;
-					if (D3D11ResourcesCache->CreateSharedResource(CoreObject->GetObjectDesc(), D3D11DevicePtr, D3D11ResourcePtr, ResourceDesc, ResourceHandle))
-					{
-						CoreObject->GetProxyData_RenderThread().ResourceHandles.Add(ResourceHandle);
+					CoreObject->GetProxyData_RenderThread().ResourceHandles.Add(ResourceHandle);
 
-						return true;
-					}
+					return true;
 				}
 			}
 		}
@@ -61,9 +54,10 @@ bool FTextureShareResource::D3D11RegisterResourceHandle(const FTextureShareCoreR
 
 bool FTextureShareResource::D3D11ReleaseTextureShareHandle()
 {
-	if (GD3D11RHI && TextureRHI.IsValid())
+	if (IsRHID3D11() && TextureRHI.IsValid())
 	{
-		if (ID3D11Texture2D* D3D11ResourcePtr = (ID3D11Texture2D*)(TextureRHI->GetNativeResource()))
+		ID3D11DynamicRHI* D3D11RHI = GetID3D11DynamicRHI();
+		if (ID3D11Texture2D* D3D11ResourcePtr = (ID3D11Texture2D*)D3D11RHI->RHIGetResource(TextureRHI))
 		{
 			TSharedPtr<ITextureShareCoreD3D11ResourcesCache, ESPMode::ThreadSafe> D3D11ResourcesCache = GetD3D11ResourcesCache();
 			if (D3D11ResourcesCache.IsValid())
diff --git a/Engine/Plugins/VirtualProduction/TextureShare/Source/TextureShare/Private/Resources/TextureShareResource_D3D12.cpp b/Engine/Plugins/VirtualProduction/TextureShare/Source/TextureShare/Private/Resources/TextureShareResource_D3D12.cpp
--- a/Engine/Plugins/VirtualProduction/TextureShare/Source/TextureShare/Private/Resources/TextureShareResource_D3D12.cpp
+++ b/Engine/Plugins/VirtualProduction/TextureShare/Source/TextureShare/Private/Resources/TextureShareResource_D3D12.cpp
@@ -8,15 +8,7 @@
 #include "ITextureShareCoreObject.h"
 #include "ITextureShareCoreD3D12ResourcesCache.h"
 
-#include "RHI.h"
-#include "RenderResource.h"
-
-#include "Windows/AllowWindowsPlatformTypes.h"
-#include "d3d12.h"
-#include "Windows/HideWindowsPlatformTypes.h"
-
-#include "D3D12RHIPrivate.h"
-#include "D3D12Util.h"
+#include "ID3D12DynamicRHI.h"
 
 //////////////////////////////////////////////////////////////////////////////////////////////
 namespace TextureShareResourceHelpers
@@ -35,22 +27,23 @@ using namespace TextureShareResourceHelpers;
 //////////////////////////////////////////////////////////////////////////////////////////////
 bool FTextureShareResource::D3D12RegisterResourceHandle(const FTextureShareCoreResourceRequest& InResourceRequest)
 {
-	if (GDynamicRHI && TextureRHI.IsValid())
+	if (IsRHID3D12() && TextureRHI.IsValid())
 	{
-		if (ID3D12Device* D3D12DevicePtr = static_cast<ID3D12Device*>(GDynamicRHI->RHIGetNativeDevice()))
+		TSharedPtr<ITextureShareCoreD3D12ResourcesCache, ESPMode::ThreadSafe> D3D12ResourcesCache = GetD3D12ResourcesCache();
+		if (D3D12ResourcesCache.IsValid())
 		{
-			if (ID3D12Resource* D3D12ResourcePtr = (ID3D12Resource*)(TextureRHI->GetNativeResource()))
+			ID3D12DynamicRHI* D3D12RHI = GetID3D12DynamicRHI();
+			ID3D12Device* D3D12DevicePtr = D3D12RHI->RHIGetDevice(0);
+			ID3D12Resource* D3D12ResourcePtr = (ID3D12Resource*)D3D12RHI->RHIGetResource(TextureRHI);
+
+			if (D3D12DevicePtr && D3D12ResourcePtr)
 			{
-				TSharedPtr<ITextureShareCoreD3D12ResourcesCache, ESPMode::ThreadSafe> D3D12ResourcesCache = GetD3D12ResourcesCache();
-				if (D3D12ResourcesCache.IsValid())
+				FTextureShareCoreResourceHandle ResourceHandle;
+				if (D3D12ResourcesCache->CreateSharedResource(CoreObject->GetObjectDesc(), D3D12DevicePtr, D3D12ResourcePtr, ResourceDesc, ResourceHandle))
 				{
-					FTextureShareCoreResourceHandle ResourceHandle;
-					if (D3D12ResourcesCache->CreateSharedResource(CoreObject->GetObjectDesc(), D3D12DevicePtr, D3D12ResourcePtr, ResourceDesc, ResourceHandle))
-					{
-						CoreObject->GetProxyData_RenderThread().ResourceHandles.Add(ResourceHandle);
+					CoreObject->GetProxyData_RenderThread().ResourceHandles.Add(ResourceHandle);
 
-						return true;
-					}
+					return true;
 				}
 			}
 		}
@@ -61,17 +54,18 @@ bool FTextureShareResource::D3D12RegisterResourceHandle(const FTextureShareCoreR
 
 bool FTextureShareResource::D3D12ReleaseTextureShareHandle()
 {
-	if (GDynamicRHI && TextureRHI.IsValid())
+	if (IsRHID3D12() && TextureRHI.IsValid())
 	{
-		if (ID3D12Device* D3D12DevicePtr = static_cast<ID3D12Device*>(GDynamicRHI->RHIGetNativeDevice()))
+		ID3D12DynamicRHI* D3D12RHI = GetID3D12DynamicRHI();
+		ID3D12Device* D3D12DevicePtr = D3D12RHI->RHIGetDevice(0);
+		ID3D12Resource* D3D12ResourcePtr = (ID3D12Resource*)D3D12RHI->RHIGetResource(TextureRHI);
+
+		if (D3D12DevicePtr && D3D12ResourcePtr)
 		{
-			if (ID3D12Resource* D3D12ResourcePtr = (ID3D12Resource*)(TextureRHI->GetNativeResource()))
+			TSharedPtr<ITextureShareCoreD3D12ResourcesCache, ESPMode::ThreadSafe> D3D12ResourcesCache = GetD3D12ResourcesCache();
+			if (D3D12ResourcesCache.IsValid())
 			{
-				TSharedPtr<ITextureShareCoreD3D12ResourcesCache, ESPMode::ThreadSafe> D3D12ResourcesCache = GetD3D12ResourcesCache();
-				if (D3D12ResourcesCache.IsValid())
-				{
-					return D3D12ResourcesCache->RemoveSharedResource(CoreObject->GetObjectDesc(), D3D12ResourcePtr);
-				}
+				return D3D12ResourcesCache->RemoveSharedResource(CoreObject->GetObjectDesc(), D3D12ResourcePtr);
 			}
 		}
 	}
diff --git a/Engine/Plugins/VirtualProduction/TextureShare/Source/TextureShare/TextureShare.Build.cs b/Engine/Plugins/VirtualProduction/TextureShare/Source/TextureShare/TextureShare.Build.cs
--- a/Engine/Plugins/VirtualProduction/TextureShare/Source/TextureShare/TextureShare.Build.cs
+++ b/Engine/Plugins/VirtualProduction/TextureShare/Source/TextureShare/TextureShare.Build.cs
@@ -11,9 +11,6 @@ public class TextureShare : ModuleRules
 		PrivateIncludePaths.AddRange(
 			new string[] {
 				Path.Combine(EngineDirectory,"Source/Runtime/Renderer/Private"),
-				Path.Combine(EngineDirectory,"Source/Runtime/Windows/D3D11RHI/Private"),
-				Path.Combine(EngineDirectory,"Source/Runtime/Windows/D3D11RHI/Private/Windows"),
-				Path.Combine(EngineDirectory,"Source/Runtime/D3D12RHI/Private"),
 				Path.Combine(EngineDirectory,"Plugins/VirtualProduction/TextureShare/Source/TextureShareCore/Private"),
 			}
 		);
@@ -33,7 +30,6 @@ public class TextureShare : ModuleRules
 				"Core",
 				"Engine",
 				"RHI",
-				"RHICore",
 				"Renderer",
 				"RenderCore",
 				"Slate",
@@ -60,20 +56,10 @@ public class TextureShare : ModuleRules
 			PublicDefinitions.Add("TEXTURESHARE_VULKAN=1");
 
 			// Support Vulkan device
-			AddEngineThirdPartyPrivateStaticDependencies(Target, "Vulkan");
 			PrivateDependencyModuleNames.AddRange(new string[] { "VulkanRHI" });
 			PrivateIncludePathModuleNames.Add("VulkanRHI");
-			PrivateIncludePaths.Add(Path.Combine(EngineDirectory, "Source/Runtime/VulkanRHI/Private"));
-			AddEngineThirdPartyPrivateStaticDependencies(Target, "Vulkan");
 
-			if (Target.IsInPlatformGroup(UnrealPlatformGroup.Windows))
-			{
-				PrivateIncludePaths.Add(Path.Combine(EngineDirectory, "Source/Runtime/VulkanRHI/Private/Windows"));
-			}
-			else if (Target.IsInPlatformGroup(UnrealPlatformGroup.Linux))
-			{
-				PrivateIncludePaths.Add(Path.Combine(EngineDirectory, "Source/Runtime/VulkanRHI/Private/Linux"));
-			}
+			AddEngineThirdPartyPrivateStaticDependencies(Target, "Vulkan");
 		}
 		else
 		{
diff --git a/Engine/Plugins/VirtualProduction/TextureShare/Source/TextureShareDisplayCluster/TextureShareDisplayCluster.Build.cs b/Engine/Plugins/VirtualProduction/TextureShare/Source/TextureShareDisplayCluster/TextureShareDisplayCluster.Build.cs
--- a/Engine/Plugins/VirtualProduction/TextureShare/Source/TextureShareDisplayCluster/TextureShareDisplayCluster.Build.cs
+++ b/Engine/Plugins/VirtualProduction/TextureShare/Source/TextureShareDisplayCluster/TextureShareDisplayCluster.Build.cs
@@ -33,7 +33,6 @@ public class TextureShareDisplayCluster : ModuleRules
 			{
 				"Engine",
 				"RHI",
-				"RHICore",
 				"Renderer",
 				"RenderCore",
 				"TextureShare",
