Fix storing a temp object pointer then using it

diff --git a/Engine/Plugins/VirtualProduction/TextureShare/Source/TextureShareCore/Private/Core/Serialize/ITextureShareCoreSerializeStream.h b/Engine/Plugins/VirtualProduction/TextureShare/Source/TextureShareCore/Private/Core/Serialize/ITextureShareCoreSerializeStream.h
--- a/Engine/Plugins/VirtualProduction/TextureShare/Source/TextureShareCore/Private/Core/Serialize/ITextureShareCoreSerializeStream.h
+++ b/Engine/Plugins/VirtualProduction/TextureShare/Source/TextureShareCore/Private/Core/Serialize/ITextureShareCoreSerializeStream.h
@@ -67,8 +67,7 @@ public:
 		{
 			if (IsWriteStream())
 			{
-				const wchar_t* WCharValue = FTCHARToWChar(*InOutData).Get();
-				SerializeData(WCharValue, sizeof(wchar_t) * StrLen);
+				SerializeData(FTCHARToWChar(*InOutData).Get(), sizeof(wchar_t) * StrLen);
 			}
 			else
 			{
