Fix lifetime issues storing a TCHAR_TO_ANSI

diff --git a/Engine/Plugins/VirtualProduction/TextureShare/Source/TextureShareCore/Private/IPC/Containers/TextureShareCoreInterprocessContainers.cpp b/Engine/Plugins/VirtualProduction/TextureShare/Source/TextureShareCore/Private/IPC/Containers/TextureShareCoreInterprocessContainers.cpp
--- a/Engine/Plugins/VirtualProduction/TextureShare/Source/TextureShareCore/Private/IPC/Containers/TextureShareCoreInterprocessContainers.cpp
+++ b/Engine/Plugins/VirtualProduction/TextureShare/Source/TextureShareCore/Private/IPC/Containers/TextureShareCoreInterprocessContainers.cpp
@@ -29,11 +29,10 @@ void FTextureShareCoreSMD5Hash::Initialize(const FString& InText)
 	Empty();
 
 	// Create MD5 Hash
-	const uint8* BlobInput = (unsigned char*)TCHAR_TO_ANSI(*InText);
 	uint64 BlobInputLen = FCString::Strlen(*InText);
 
 	FMD5 Md5Gen;
-	Md5Gen.Update(BlobInput, BlobInputLen);
+	Md5Gen.Update((unsigned char*)TCHAR_TO_ANSI(*InText), BlobInputLen);
 	Md5Gen.Final(MD5Digest);
 }
 
