Dont store temp pointers, fixes dangling pointer usage

diff --git a/Engine/Plugins/VirtualProduction/TextureShare/Source/TextureShareCore/Private/IPC/Containers/TextureShareCoreInterprocessContainers.cpp b/Engine/Plugins/VirtualProduction/TextureShare/Source/TextureShareCore/Private/IPC/Containers/TextureShareCoreInterprocessContainers.cpp
--- a/Engine/Plugins/VirtualProduction/TextureShare/Source/TextureShareCore/Private/IPC/Containers/TextureShareCoreInterprocessContainers.cpp
+++ b/Engine/Plugins/VirtualProduction/TextureShare/Source/TextureShareCore/Private/IPC/Containers/TextureShareCoreInterprocessContainers.cpp
@@ -69,7 +69,8 @@ void FTextureShareCoreStringHash::Initialize(const FString& InText)
 	Hash.Initialize(InText);
 
 	// Copy string
-	const wchar_t* WCharValue = FTCHARToWChar(*InText).Get();
+	FTCHARToWChar WCharString(*InText);
+	const wchar_t* WCharValue = WCharString.Get();
 	if (WCharValue && WCharValue[0])
 	{
 		for (int32 CharIt = 0; CharIt < (MaxStringLength-1); CharIt++)
diff --git a/Engine/Plugins/VirtualProduction/TextureShare/Source/TextureShareCore/Private/IPC/TextureShareCoreInterprocessEvent.cpp b/Engine/Plugins/VirtualProduction/TextureShare/Source/TextureShareCore/Private/IPC/TextureShareCoreInterprocessEvent.cpp
--- a/Engine/Plugins/VirtualProduction/TextureShare/Source/TextureShareCore/Private/IPC/TextureShareCoreInterprocessEvent.cpp
+++ b/Engine/Plugins/VirtualProduction/TextureShare/Source/TextureShareCore/Private/IPC/TextureShareCoreInterprocessEvent.cpp
@@ -18,14 +18,13 @@ namespace TextureShareCoreInterprocessEventHelper
 	static TSharedPtr<FEvent, ESPMode::ThreadSafe> CreateNamedInterprocessEvent(const FString& EventName, const bool bOpenExistEvent, const void* InSecurityAttributes)
 	{
 		const FString GlobalEventName = FString::Printf(TEXT("Global\\%s"), *EventName);
-		const LPCSTR EventNameStr = TCHAR_TO_ANSI(*GlobalEventName);
 
 		HANDLE EventHandle = NULL;
 
 		if (bOpenExistEvent)
 		{
 			const DWORD AccessRights = EVENT_ALL_ACCESS | EVENT_MODIFY_STATE;//SYNCHRONIZE | EVENT_MODIFY_STATE;
-			EventHandle = OpenEventA(AccessRights, false, EventNameStr);
+			EventHandle = OpenEventA(AccessRights, false, TCHAR_TO_ANSI(*GlobalEventName));
 			if (NULL == EventHandle)
 			{
 				UE_LOG(LogTextureShareCore, Warning, TEXT("OpenEvent(AccessRights=0x%08x, bInherit=false, Name='%s') failed with LastError = %d"), AccessRights, *EventName, GetLastError());
@@ -36,7 +35,7 @@ namespace TextureShareCoreInterprocessEventHelper
 		else
 		{
 			// Create new event
-			EventHandle = CreateEventA((SECURITY_ATTRIBUTES*)InSecurityAttributes, true, false, EventNameStr);
+			EventHandle = CreateEventA((SECURITY_ATTRIBUTES*)InSecurityAttributes, true, false, TCHAR_TO_ANSI(*GlobalEventName));
 			if (NULL == EventHandle)
 			{
 				UE_LOG(LogTextureShareCore, Warning, TEXT("CreateEvent(NULL, bInherit=false, Name='%s') failed with LastError = %d"), *EventName, GetLastError());
