TextureShare: added missing API callbacks for ndisplay texture shared object custom data.

diff --git a/Engine/Plugins/VirtualProduction/TextureShare/Source/TextureShare/Private/ITextureShareAPI.h b/Engine/Plugins/VirtualProduction/TextureShare/Source/TextureShare/Private/ITextureShareAPI.h
--- a/Engine/Plugins/VirtualProduction/TextureShare/Source/TextureShare/Private/ITextureShareAPI.h
+++ b/Engine/Plugins/VirtualProduction/TextureShare/Source/TextureShare/Private/ITextureShareAPI.h
@@ -6,6 +6,7 @@
 class ITextureShareObject;
 class ITextureShareObjectProxy;
 class UWorld;
+class ITextureShareCallbacks;
 
 /**
  * TextureShare API
@@ -103,4 +104,12 @@ public:
 	 * @param InProcessId - Custom local process name
 	 */
 	virtual void SetProcessName(const FString& InProcessId) = 0;
+
+public:
+	/**
+	* Access to the TextureShare callbacks API
+	*
+	* @return Callbacks API
+	*/
+	virtual ITextureShareCallbacks& GetCallbacks() = 0;
 };
diff --git a/Engine/Plugins/VirtualProduction/TextureShare/Source/TextureShare/Private/ITextureShareCallbacks.h b/Engine/Plugins/VirtualProduction/TextureShare/Source/TextureShare/Private/ITextureShareCallbacks.h
new file mode 100644
--- /dev/null
+++ b/Engine/Plugins/VirtualProduction/TextureShare/Source/TextureShare/Private/ITextureShareCallbacks.h
@@ -0,0 +1,57 @@
+// Copyright Epic Games, Inc. All Rights Reserved.
+
+#pragma once
+
+#include "CoreMinimal.h"
+
+#include "Containers/TextureShareContainers.h"
+#include "Containers/TextureShareCoreEnums.h"
+
+class ITextureShareObject;
+class ITextureShareObjectProxy;
+class FRHICommandListImmediate;
+
+/**
+ * TextureShare callbacks API
+ */
+class ITextureShareCallbacks
+{
+public:
+	static ITextureShareCallbacks& Get();
+
+	virtual ~ITextureShareCallbacks() = default;
+
+public:
+	/** Called on session start **/
+	DECLARE_EVENT_OneParam(ITextureShareCallbacks, FTextureShareBeginSessionEvent, ITextureShareObject&);
+	virtual FTextureShareBeginSessionEvent& OnTextureShareBeginSession() = 0;
+
+	/** Called on session end **/
+	DECLARE_EVENT_OneParam(ITextureShareCallbacks, FTextureShareEndSessionEvent, ITextureShareObject&);
+	virtual FTextureShareEndSessionEvent& OnTextureShareEndSession() = 0;
+
+	/** Called on begin frame sync **/
+	DECLARE_EVENT_OneParam(ITextureShareCallbacks, FTextureShareBeginFrameSyncEvent, ITextureShareObject&);
+	virtual FTextureShareBeginFrameSyncEvent& OnTextureShareBeginFrameSync() = 0;
+
+	/** Called on end frame sync **/
+	DECLARE_EVENT_OneParam(ITextureShareCallbacks, FTextureShareEndFrameSyncEvent, ITextureShareObject&);
+	virtual FTextureShareEndFrameSyncEvent& OnTextureShareEndFrameSync() = 0;
+
+	/** Called on frame sync **/
+	DECLARE_EVENT_TwoParams(ITextureShareCallbacks, FTextureShareFrameSyncEvent, ITextureShareObject&, const ETextureShareSyncStep);
+	virtual FTextureShareFrameSyncEvent& OnTextureShareFrameSync() = 0;
+
+
+	/** Called on begin frame sync on render thread **/
+	DECLARE_EVENT_TwoParams(ITextureShareCallbacks, FTextureShareBeginFrameSyncEvent_RenderThread, FRHICommandListImmediate&, const ITextureShareObjectProxy&);
+	virtual FTextureShareBeginFrameSyncEvent_RenderThread& OnTextureShareBeginFrameSyncEvent_RenderThread() = 0;
+
+	/** Called on end frame sync on render thread **/
+	DECLARE_EVENT_TwoParams(ITextureShareCallbacks, FTextureShareEndFrameSyncEvent_RenderThread, FRHICommandListImmediate&, const ITextureShareObjectProxy&);
+	virtual FTextureShareEndFrameSyncEvent_RenderThread& OnTextureShareEndFrameSyncEvent_RenderThread() = 0;
+
+	/** Called on frame sync on render thread **/
+	DECLARE_EVENT_ThreeParams(ITextureShareCallbacks, FTextureShareFrameSyncEvent_RenderThread, FRHICommandListImmediate&, const ITextureShareObjectProxy&, const ETextureShareSyncStep);
+	virtual FTextureShareFrameSyncEvent_RenderThread& OnTextureShareFrameSyncEvent_RenderThread() = 0;
+};
diff --git a/Engine/Plugins/VirtualProduction/TextureShare/Source/TextureShare/Private/Module/TextureShareAPI.h b/Engine/Plugins/VirtualProduction/TextureShare/Source/TextureShare/Private/Module/TextureShareAPI.h
--- a/Engine/Plugins/VirtualProduction/TextureShare/Source/TextureShare/Private/Module/TextureShareAPI.h
+++ b/Engine/Plugins/VirtualProduction/TextureShare/Source/TextureShare/Private/Module/TextureShareAPI.h
@@ -2,6 +2,7 @@
 
 #pragma once
 #include "ITextureShareAPI.h"
+#include "TextureShareCallbacks.h"
 
 #include "Object/TextureShareObject.h"
 #include "Object/TextureShareObjectProxy.h"
@@ -34,6 +35,11 @@ public:
 
 	virtual void OnWorldBeginPlay(UWorld& InWorld) override;
 	virtual void OnWorldEndPlay(UWorld& InWorld) override;
+
+	virtual ITextureShareCallbacks& GetCallbacks() override
+	{
+		return Callbacks;
+	}
 	//~~ ITextureShare
 
 private:
@@ -62,4 +68,6 @@ private:
 	FDelegateHandle OnBackBufferReadyToPresentHandle;
 
 	mutable FCriticalSection ThreadDataCS;
+
+	FTextureShareCallbacks Callbacks;
 };
diff --git a/Engine/Plugins/VirtualProduction/TextureShare/Source/TextureShare/Private/Module/TextureShareCallbacks.cpp b/Engine/Plugins/VirtualProduction/TextureShare/Source/TextureShare/Private/Module/TextureShareCallbacks.cpp
new file mode 100644
--- /dev/null
+++ b/Engine/Plugins/VirtualProduction/TextureShare/Source/TextureShare/Private/Module/TextureShareCallbacks.cpp
@@ -0,0 +1,14 @@
+// Copyright Epic Games, Inc. All Rights Reserved.
+
+#include "Module/TextureShareCallbacks.h"
+
+#include "ITextureShare.h"
+#include "ITextureShareAPI.h"
+#include "ITextureShareCallbacks.h"
+
+ITextureShareCallbacks& ITextureShareCallbacks::Get()
+{
+	static ITextureShareCallbacks& TSCallbacksAPI = ITextureShare::Get().GetTextureShareAPI().GetCallbacks();
+
+	return TSCallbacksAPI;
+}
diff --git a/Engine/Plugins/VirtualProduction/TextureShare/Source/TextureShare/Private/Module/TextureShareCallbacks.h b/Engine/Plugins/VirtualProduction/TextureShare/Source/TextureShare/Private/Module/TextureShareCallbacks.h
new file mode 100644
--- /dev/null
+++ b/Engine/Plugins/VirtualProduction/TextureShare/Source/TextureShare/Private/Module/TextureShareCallbacks.h
@@ -0,0 +1,69 @@
+// Copyright Epic Games, Inc. All Rights Reserved.
+
+#pragma once
+
+#include "ITextureShareCallbacks.h"
+
+/**
+ * TextureShare callbacks API implementation
+ */
+class FTextureShareCallbacks :
+	public  ITextureShareCallbacks
+{
+public:
+	virtual ~FTextureShareCallbacks() = default;
+
+public:
+	virtual FTextureShareBeginSessionEvent& OnTextureShareBeginSession() override
+	{
+		return TextureShareBeginSessionEvent;
+	}
+	
+	virtual FTextureShareEndSessionEvent& OnTextureShareEndSession() override
+	{
+		return TextureShareEndSessionEvent;
+	}
+
+
+	virtual FTextureShareBeginFrameSyncEvent& OnTextureShareBeginFrameSync() override
+	{
+		return TextureShareBeginFrameSyncEvent;
+	}
+
+	virtual FTextureShareEndFrameSyncEvent& OnTextureShareEndFrameSync() override
+	{
+		return TextureShareEndFrameSyncEvent;
+	}
+
+	virtual FTextureShareFrameSyncEvent& OnTextureShareFrameSync() override
+	{
+		return TextureShareFrameSyncEvent;
+	}
+
+	virtual FTextureShareBeginFrameSyncEvent_RenderThread& OnTextureShareBeginFrameSyncEvent_RenderThread() override
+	{
+		return TextureShareBeginFrameSyncEvent_RenderThread;
+	}
+
+	virtual FTextureShareEndFrameSyncEvent_RenderThread& OnTextureShareEndFrameSyncEvent_RenderThread() override
+	{
+		return TextureShareEndFrameSyncEvent_RenderThread;
+	}
+
+	virtual FTextureShareFrameSyncEvent_RenderThread& OnTextureShareFrameSyncEvent_RenderThread() override
+	{
+		return TextureShareFrameSyncEvent_RenderThread;
+	}
+
+private:
+	FTextureShareBeginSessionEvent     TextureShareBeginSessionEvent;
+	FTextureShareEndSessionEvent       TextureShareEndSessionEvent;
+
+	FTextureShareBeginFrameSyncEvent   TextureShareBeginFrameSyncEvent;
+	FTextureShareEndFrameSyncEvent     TextureShareEndFrameSyncEvent;
+	FTextureShareFrameSyncEvent        TextureShareFrameSyncEvent;
+
+	FTextureShareBeginFrameSyncEvent_RenderThread TextureShareBeginFrameSyncEvent_RenderThread;
+	FTextureShareEndFrameSyncEvent_RenderThread   TextureShareEndFrameSyncEvent_RenderThread;
+	FTextureShareFrameSyncEvent_RenderThread      TextureShareFrameSyncEvent_RenderThread;
+};
diff --git a/Engine/Plugins/VirtualProduction/TextureShare/Source/TextureShare/Private/Object/TextureShareObject.cpp b/Engine/Plugins/VirtualProduction/TextureShare/Source/TextureShare/Private/Object/TextureShareObject.cpp
--- a/Engine/Plugins/VirtualProduction/TextureShare/Source/TextureShare/Private/Object/TextureShareObject.cpp
+++ b/Engine/Plugins/VirtualProduction/TextureShare/Source/TextureShare/Private/Object/TextureShareObject.cpp
@@ -5,6 +5,7 @@
 #include "Module/TextureShareLog.h"
 #include "Misc/TextureShareStrings.h"
 
+#include "ITextureShareCallbacks.h"
 #include "ITextureShareCoreObject.h"
 
 //////////////////////////////////////////////////////////////////////////////////////////////
@@ -71,6 +72,11 @@ bool FTextureShareObject::BeginSession()
 
 		FTextureShareObjectProxy::BeginSession_GameThread(*this);
 
+		if (ITextureShareCallbacks::Get().OnTextureShareBeginSession().IsBound())
+		{
+			ITextureShareCallbacks::Get().OnTextureShareBeginSession().Broadcast(*this);
+		}
+
 		return true;
 	}
 
@@ -85,6 +91,11 @@ bool FTextureShareObject::EndSession()
 
 		FTextureShareObjectProxy::EndSession_GameThread(*this);
 
+		if (ITextureShareCallbacks::Get().OnTextureShareEndSession().IsBound())
+		{
+			ITextureShareCallbacks::Get().OnTextureShareEndSession().Broadcast(*this);
+		}
+
 		return true;
 	}
 
@@ -105,6 +116,11 @@ bool FTextureShareObject::BeginFrameSync()
 		{
 			bFrameSyncActive = true;
 
+			if (ITextureShareCallbacks::Get().OnTextureShareBeginFrameSync().IsBound())
+			{
+				ITextureShareCallbacks::Get().OnTextureShareBeginFrameSync().Broadcast(*this);
+			}
+
 			return true;
 		}
 	}
@@ -133,6 +149,11 @@ bool FTextureShareObject::EndFrameSync(FViewport* InViewport)
 	// Game thread data now sent to the proxy.Game-thread data can now be cleared
 	const bool bResult = CoreObject->EndFrameSync();
 
+	if (ITextureShareCallbacks::Get().OnTextureShareEndFrameSync().IsBound())
+	{
+		ITextureShareCallbacks::Get().OnTextureShareEndFrameSync().Broadcast(*this);
+	}
+
 	CoreObject->UnlockThreadMutex(ETextureShareThreadMutex::RenderingThread);
 
 	return bResult;
@@ -140,8 +161,15 @@ bool FTextureShareObject::EndFrameSync(FViewport* InViewport)
 
 bool FTextureShareObject::FrameSync(const ETextureShareSyncStep InSyncStep)
 {
+	TRACE_CPUPROFILER_EVENT_SCOPE(TextureShare::FrameSync);
+
 	if (IsFrameSyncActive() && CoreObject->FrameSync(InSyncStep))
 	{
+		if (ITextureShareCallbacks::Get().OnTextureShareFrameSync().IsBound())
+		{
+			ITextureShareCallbacks::Get().OnTextureShareFrameSync().Broadcast(*this, InSyncStep);
+		}
+
 		return true;
 	}
 
diff --git a/Engine/Plugins/VirtualProduction/TextureShare/Source/TextureShare/Private/Object/TextureShareObjectProxy.cpp b/Engine/Plugins/VirtualProduction/TextureShare/Source/TextureShare/Private/Object/TextureShareObjectProxy.cpp
--- a/Engine/Plugins/VirtualProduction/TextureShare/Source/TextureShare/Private/Object/TextureShareObjectProxy.cpp
+++ b/Engine/Plugins/VirtualProduction/TextureShare/Source/TextureShare/Private/Object/TextureShareObjectProxy.cpp
@@ -9,6 +9,7 @@
 #include "Module/TextureShareLog.h"
 #include "Misc/TextureShareStrings.h"
 
+#include "ITextureShareCallbacks.h"
 #include "ITextureShareCore.h"
 
 #include "RenderGraphUtils.h"
@@ -84,6 +85,11 @@ bool FTextureShareObjectProxy::BeginFrameSync_RenderThread(FRHICommandListImmedi
 
 			bFrameProxySyncActive = true;
 
+			if (ITextureShareCallbacks::Get().OnTextureShareBeginFrameSyncEvent_RenderThread().IsBound())
+			{
+				ITextureShareCallbacks::Get().OnTextureShareBeginFrameSyncEvent_RenderThread().Broadcast(RHICmdList, *this);
+			}
+
 			return true;
 		}
 	}
@@ -102,15 +108,27 @@ bool FTextureShareObjectProxy::EndFrameSync_RenderThread(FRHICommandListImmediat
 
 	bFrameProxySyncActive = false;
 
+	if (ITextureShareCallbacks::Get().OnTextureShareEndFrameSyncEvent_RenderThread().IsBound())
+	{
+		ITextureShareCallbacks::Get().OnTextureShareEndFrameSyncEvent_RenderThread().Broadcast(RHICmdList, *this);
+	}
+
 	CoreObject->UnlockThreadMutex(ETextureShareThreadMutex::GameThread);
 
 	return bResult;
 }
 
+DECLARE_GPU_STAT_NAMED(TextureShareObjectProxyFrameSync, TEXT("TextureShare::FrameSync_RenderThread"));
+
 bool FTextureShareObjectProxy::FrameSync_RenderThread(FRHICommandListImmediate& RHICmdList, const ETextureShareSyncStep InSyncStep) const
 {
 	check(IsInRenderingThread());
 
+	SCOPED_GPU_STAT(RHICmdList, TextureShareObjectProxyFrameSync);
+	SCOPED_DRAW_EVENT(RHICmdList, TextureShareObjectProxyFrameSync);
+
+	TRACE_CPUPROFILER_EVENT_SCOPE(TextureShare::FrameSync_RenderThread);
+
 	if (IsFrameSyncActive_RenderThread())
 	{
 		if (ResourcesProxy.IsValid())
@@ -137,6 +155,11 @@ bool FTextureShareObjectProxy::FrameSync_RenderThread(FRHICommandListImmediate&
 				ResourcesProxy->RunReceiveResources_RenderThread(RHICmdList, InSyncStep);
 			}
 
+			if (ITextureShareCallbacks::Get().OnTextureShareFrameSyncEvent_RenderThread().IsBound())
+			{
+				ITextureShareCallbacks::Get().OnTextureShareFrameSyncEvent_RenderThread().Broadcast(RHICmdList, *this, InSyncStep);
+			}
+
 			return true;
 		}
 	}
@@ -242,6 +265,8 @@ bool FTextureShareObjectProxy::ShareResource_RenderThread(FRDGBuilder& GraphBuil
 }
 
 //////////////////////////////////////////////////////////////////////////////////////////////
+DECLARE_GPU_STAT_NAMED(TextureShareObjectProxyShareResource, TEXT("TextureShare::ShareResource_RenderThread"));
+
 bool FTextureShareObjectProxy::ShareResource_RenderThread(FRHICommandListImmediate& RHICmdList, const FTextureShareCoreResourceDesc& InResourceDesc, FRHITexture* InTexture, const int32 InTextureGPUIndex, const FIntRect* InTextureRect) const
 {
 	check(IsInRenderingThread());
@@ -253,6 +278,12 @@ bool FTextureShareObjectProxy::ShareResource_RenderThread(FRHICommandListImmedia
 		{
 			if (FTextureShareResource* SharedResource = ResourcesProxy->GetSharedTexture_RenderThread(RHICmdList, CoreObject, InTexture, *ExistResourceRequest))
 			{
+
+				SCOPED_GPU_STAT(RHICmdList, TextureShareObjectProxyShareResource);
+				SCOPED_DRAW_EVENT(RHICmdList, TextureShareObjectProxyShareResource);
+
+				TRACE_CPUPROFILER_EVENT_SCOPE(TextureShare::ShareResource_RenderThread);
+
 				switch (ExistResourceRequest->ResourceDesc.OperationType)
 				{
 					// Remote process request read texture, send it
diff --git a/Engine/Plugins/VirtualProduction/TextureShare/Source/TextureShareCore/Private/Core/TextureShareCore.h b/Engine/Plugins/VirtualProduction/TextureShare/Source/TextureShareCore/Private/Core/TextureShareCore.h
--- a/Engine/Plugins/VirtualProduction/TextureShare/Source/TextureShareCore/Private/Core/TextureShareCore.h
+++ b/Engine/Plugins/VirtualProduction/TextureShare/Source/TextureShareCore/Private/Core/TextureShareCore.h
@@ -2,6 +2,8 @@
 
 #pragma once
 #include "ITextureShareCoreAPI.h"
+#include "TextureShareCoreCallbacks.h"
+
 #include "Containers/TextureShareCoreContainers.h"
 
 class FTextureShareCoreObject;
@@ -44,6 +46,11 @@ public:
 
 	virtual void SetProcessName(const FString& InProcessId) override;
 	virtual bool SetProcessDeviceType(const ETextureShareDeviceType InDeviceType) override;
+
+	virtual ITextureShareCoreCallbacks& GetCallbacks() override
+	{
+		return Callbacks;
+	}
 	//~~ITextureShareCoreAPI
 
 	bool RemoveCachedResources(const FTextureShareCoreObjectDesc& ObjectDesc);
@@ -86,4 +93,6 @@ private:
 
 	// TS Objects
 	TMap<FString, TSharedPtr<FTextureShareCoreObject, ESPMode::ThreadSafe>> Objects;
+
+	FTextureShareCoreCallbacks Callbacks;
 };
diff --git a/Engine/Plugins/VirtualProduction/TextureShare/Source/TextureShareCore/Private/Core/TextureShareCoreCallbacks.cpp b/Engine/Plugins/VirtualProduction/TextureShare/Source/TextureShareCore/Private/Core/TextureShareCoreCallbacks.cpp
new file mode 100644
--- /dev/null
+++ b/Engine/Plugins/VirtualProduction/TextureShare/Source/TextureShareCore/Private/Core/TextureShareCoreCallbacks.cpp
@@ -0,0 +1,14 @@
+// Copyright Epic Games, Inc. All Rights Reserved.
+
+#include "Core/TextureShareCoreCallbacks.h"
+
+#include "ITextureShareCore.h"
+#include "ITextureShareCoreAPI.h"
+#include "ITextureShareCoreCallbacks.h"
+
+ITextureShareCoreCallbacks& ITextureShareCoreCallbacks::Get()
+{
+	static ITextureShareCoreCallbacks& TSCallbacksAPI = ITextureShareCore::Get().GetTextureShareCoreAPI().GetCallbacks();
+
+	return TSCallbacksAPI;
+}
diff --git a/Engine/Plugins/VirtualProduction/TextureShare/Source/TextureShareCore/Private/Core/TextureShareCoreCallbacks.h b/Engine/Plugins/VirtualProduction/TextureShare/Source/TextureShareCore/Private/Core/TextureShareCoreCallbacks.h
new file mode 100644
--- /dev/null
+++ b/Engine/Plugins/VirtualProduction/TextureShare/Source/TextureShareCore/Private/Core/TextureShareCoreCallbacks.h
@@ -0,0 +1,67 @@
+// Copyright Epic Games, Inc. All Rights Reserved.
+
+#pragma once
+
+#include "ITextureShareCoreCallbacks.h"
+
+/**
+ * TextureShareCore callbacks API implementation
+ */
+class FTextureShareCoreCallbacks :
+	public  ITextureShareCoreCallbacks
+{
+public:
+	virtual ~FTextureShareCoreCallbacks() = default;
+
+public:
+	virtual FTextureShareCoreBeginSessionEvent& OnTextureShareCoreBeginSession() override
+	{
+		return TextureShareCoreBeginSessionEvent;
+	}
+
+	virtual FTextureShareCoreEndSessionEvent& OnTextureShareCoreEndSession() override
+	{
+		return TextureShareCoreEndSessionEvent;
+	}
+
+	virtual FTextureShareCoreBeginFrameSyncEvent& OnTextureShareCoreBeginFrameSync() override
+	{
+		return TextureShareCoreBeginFrameSyncEvent;
+	}
+
+	virtual FTextureShareCoreEndFrameSyncEvent& OnTextureShareCoreEndFrameSync() override
+	{
+		return TextureShareCoreEndFrameSyncEvent;
+	}
+
+	virtual FTextureShareCoreFrameSyncEvent& OnTextureShareCoreFrameSync() override
+	{
+		return TextureShareCoreFrameSyncEvent;
+	}
+
+	virtual FTextureShareCoreBeginFrameSyncEvent_RenderThread& OnTextureShareCoreBeginFrameSync_RenderThread() override
+	{
+		return TextureShareCoreBeginFrameSyncEvent_RenderThread;
+	}
+
+	virtual FTextureShareCoreEndFrameSyncEvent_RenderThread& OnTextureShareCoreEndFrameSync_RenderThread() override
+	{
+		return TextureShareCoreEndFrameSyncEvent_RenderThread;
+	}
+
+	virtual FTextureShareCoreFrameSyncEvent_RenderThread& OnTextureShareCoreFrameSync_RenderThread() override
+	{
+		return TextureShareCoreFrameSyncEvent_RenderThread;
+	}
+
+private:
+	FTextureShareCoreBeginSessionEvent     TextureShareCoreBeginSessionEvent;
+	FTextureShareCoreEndSessionEvent       TextureShareCoreEndSessionEvent;
+	FTextureShareCoreBeginFrameSyncEvent   TextureShareCoreBeginFrameSyncEvent;
+	FTextureShareCoreEndFrameSyncEvent     TextureShareCoreEndFrameSyncEvent;
+	FTextureShareCoreFrameSyncEvent        TextureShareCoreFrameSyncEvent;
+
+	FTextureShareCoreBeginFrameSyncEvent_RenderThread TextureShareCoreBeginFrameSyncEvent_RenderThread;
+	FTextureShareCoreEndFrameSyncEvent_RenderThread   TextureShareCoreEndFrameSyncEvent_RenderThread;
+	FTextureShareCoreFrameSyncEvent_RenderThread      TextureShareCoreFrameSyncEvent_RenderThread;
+};
diff --git a/Engine/Plugins/VirtualProduction/TextureShare/Source/TextureShareCore/Private/ITextureShareCoreAPI.h b/Engine/Plugins/VirtualProduction/TextureShare/Source/TextureShareCore/Private/ITextureShareCoreAPI.h
--- a/Engine/Plugins/VirtualProduction/TextureShare/Source/TextureShareCore/Private/ITextureShareCoreAPI.h
+++ b/Engine/Plugins/VirtualProduction/TextureShare/Source/TextureShareCore/Private/ITextureShareCoreAPI.h
@@ -8,6 +8,8 @@
 
 #include "Containers/TextureShareCoreContainers_ObjectDesc.h"
 
+class ITextureShareCoreCallbacks;
+
 /**
  * TextureShareCore API
  */
@@ -96,4 +98,12 @@ public:
 	 * @param InMilisecondsTimeout- Timeout to remove resource
 	 */
 	virtual void RemoveUnusedResources(const uint32 InMilisecondsTimeOut) = 0;
+
+public:
+	/**
+	* Access to the TextureShareCore callbacks API
+	*
+	* @return Callbacks API
+	*/
+	virtual ITextureShareCoreCallbacks& GetCallbacks() = 0;
 };
diff --git a/Engine/Plugins/VirtualProduction/TextureShare/Source/TextureShareCore/Private/ITextureShareCoreCallbacks.h b/Engine/Plugins/VirtualProduction/TextureShare/Source/TextureShareCore/Private/ITextureShareCoreCallbacks.h
new file mode 100644
--- /dev/null
+++ b/Engine/Plugins/VirtualProduction/TextureShare/Source/TextureShareCore/Private/ITextureShareCoreCallbacks.h
@@ -0,0 +1,54 @@
+// Copyright Epic Games, Inc. All Rights Reserved.
+
+#pragma once
+#include "CoreMinimal.h"
+#include "Containers/TextureShareCoreEnums.h"
+
+class ITextureShareCoreObject;
+class FRHICommandListImmediate;
+
+/**
+ * TextureShareCore callbacks API
+ */
+class ITextureShareCoreCallbacks
+{
+public:
+	static ITextureShareCoreCallbacks& Get();
+
+	virtual ~ITextureShareCoreCallbacks() = default;
+
+public:
+	/** Called on session start **/
+	DECLARE_EVENT_OneParam(ITextureShareCoreCallbacks, FTextureShareCoreBeginSessionEvent, ITextureShareCoreObject&);
+	virtual FTextureShareCoreBeginSessionEvent& OnTextureShareCoreBeginSession() = 0;
+
+	/** Called on session end **/
+	DECLARE_EVENT_OneParam(ITextureShareCoreCallbacks, FTextureShareCoreEndSessionEvent, ITextureShareCoreObject&);
+	virtual FTextureShareCoreEndSessionEvent& OnTextureShareCoreEndSession() = 0;
+
+
+	/** Called on begin frame sync **/
+	DECLARE_EVENT_OneParam(ITextureShareCoreCallbacks, FTextureShareCoreBeginFrameSyncEvent, ITextureShareCoreObject&);
+	virtual FTextureShareCoreBeginFrameSyncEvent& OnTextureShareCoreBeginFrameSync() = 0;
+
+	/** Called on end frame sync **/
+	DECLARE_EVENT_OneParam(ITextureShareCoreCallbacks, FTextureShareCoreEndFrameSyncEvent, ITextureShareCoreObject&);
+	virtual FTextureShareCoreEndFrameSyncEvent& OnTextureShareCoreEndFrameSync() = 0;
+
+	/** Called on frame sync **/
+	DECLARE_EVENT_TwoParams(ITextureShareCoreCallbacks, FTextureShareCoreFrameSyncEvent, ITextureShareCoreObject&, const ETextureShareSyncStep);
+	virtual FTextureShareCoreFrameSyncEvent& OnTextureShareCoreFrameSync() = 0;
+
+
+	/** Called on begin frame sync on render thread **/
+	DECLARE_EVENT_OneParam(ITextureShareCoreCallbacks, FTextureShareCoreBeginFrameSyncEvent_RenderThread, ITextureShareCoreObject&);
+	virtual FTextureShareCoreBeginFrameSyncEvent_RenderThread& OnTextureShareCoreBeginFrameSync_RenderThread() = 0;
+
+	/** Called on end frame sync on render thread **/
+	DECLARE_EVENT_OneParam(ITextureShareCoreCallbacks, FTextureShareCoreEndFrameSyncEvent_RenderThread, ITextureShareCoreObject&);
+	virtual FTextureShareCoreEndFrameSyncEvent_RenderThread& OnTextureShareCoreEndFrameSync_RenderThread() = 0;
+
+	/** Called on frame sync on render thread **/
+	DECLARE_EVENT_TwoParams(ITextureShareCoreCallbacks, FTextureShareCoreFrameSyncEvent_RenderThread, ITextureShareCoreObject&, const ETextureShareSyncStep);
+	virtual FTextureShareCoreFrameSyncEvent_RenderThread& OnTextureShareCoreFrameSync_RenderThread() = 0;
+};
diff --git a/Engine/Plugins/VirtualProduction/TextureShare/Source/TextureShareCore/Private/Object/TextureShareCoreObject_Connect.cpp b/Engine/Plugins/VirtualProduction/TextureShare/Source/TextureShareCore/Private/Object/TextureShareCoreObject_Connect.cpp
--- a/Engine/Plugins/VirtualProduction/TextureShare/Source/TextureShareCore/Private/Object/TextureShareCoreObject_Connect.cpp
+++ b/Engine/Plugins/VirtualProduction/TextureShare/Source/TextureShareCore/Private/Object/TextureShareCoreObject_Connect.cpp
@@ -24,6 +24,8 @@ using namespace TextureShareCoreHelpers;
 //////////////////////////////////////////////////////////////////////////////////////////////
 bool FTextureShareCoreObject::TryFrameProcessesConnection(FTextureShareCoreInterprocessMemory& InterprocessMemory, FTextureShareCoreInterprocessObject& LocalObject)
 {
+	TRACE_CPUPROFILER_EVENT_SCOPE(TextureShareCore::Connect);
+
 	// Mark last acces time for current process (this value used by other processes to detect die)
 	LocalObject.Sync.UpdateLastAccessTime();
 
diff --git a/Engine/Plugins/VirtualProduction/TextureShare/Source/TextureShareCore/Private/Object/TextureShareCoreObject_Session.cpp b/Engine/Plugins/VirtualProduction/TextureShare/Source/TextureShareCore/Private/Object/TextureShareCoreObject_Session.cpp
--- a/Engine/Plugins/VirtualProduction/TextureShare/Source/TextureShareCore/Private/Object/TextureShareCoreObject_Session.cpp
+++ b/Engine/Plugins/VirtualProduction/TextureShare/Source/TextureShareCore/Private/Object/TextureShareCoreObject_Session.cpp
@@ -13,6 +13,8 @@
 #include "Module/TextureShareCoreModule.h"
 #include "Module/TextureShareCoreLog.h"
 
+#include "ITextureShareCoreCallbacks.h"
+
 #include "Misc/ScopeLock.h"
 
 //////////////////////////////////////////////////////////////////////////////////////////////
@@ -38,6 +40,11 @@ bool FTextureShareCoreObject::BeginSession()
 				// New IPC object created
 				bSessionActive = true;
 
+				if (ITextureShareCoreCallbacks::Get().OnTextureShareCoreBeginSession().IsBound())
+				{
+					ITextureShareCoreCallbacks::Get().OnTextureShareCoreBeginSession().Broadcast(*this);
+				}
+
 				UE_TS_LOG(LogTextureShareCoreObject, Log, TEXT("%s:BeginSession()"), *GetName());
 			}
 		}
@@ -67,6 +74,11 @@ bool FTextureShareCoreObject::EndSession()
 					// Reset IPC object to defaults
 					InterprocessObject->Release();
 
+					if (ITextureShareCoreCallbacks::Get().OnTextureShareCoreEndSession().IsBound())
+					{
+						ITextureShareCoreCallbacks::Get().OnTextureShareCoreEndSession().Broadcast(*this);
+					}
+
 					UE_TS_LOG(LogTextureShareCoreObject, Log, TEXT("%s:EndSession()"), *GetName());
 
 					bResult = true;
diff --git a/Engine/Plugins/VirtualProduction/TextureShare/Source/TextureShareCore/Private/Object/TextureShareCoreObject_Sync.cpp b/Engine/Plugins/VirtualProduction/TextureShare/Source/TextureShareCore/Private/Object/TextureShareCoreObject_Sync.cpp
--- a/Engine/Plugins/VirtualProduction/TextureShare/Source/TextureShareCore/Private/Object/TextureShareCoreObject_Sync.cpp
+++ b/Engine/Plugins/VirtualProduction/TextureShare/Source/TextureShareCore/Private/Object/TextureShareCoreObject_Sync.cpp
@@ -15,6 +15,8 @@
 
 #include "Misc/ScopeLock.h"
 
+#include "ITextureShareCoreCallbacks.h"
+
 //////////////////////////////////////////////////////////////////////////////////////////////
 using namespace TextureShareCoreHelpers;
 
@@ -23,6 +25,8 @@ using namespace TextureShareCoreHelpers;
 //////////////////////////////////////////////////////////////////////////////////////////////
 bool FTextureShareCoreObject::FrameSync(const ETextureShareSyncStep InSyncStep)
 {
+	TRACE_CPUPROFILER_EVENT_SCOPE(TextureShareCore::FrameSync);
+
 	UE_TS_LOG(LogTextureShareCoreObjectSync, Log, TEXT(">>>> %s:FrameSync(%s)"), *GetName(), GetTEXT(InSyncStep));
 
 	if (IsFrameSyncActive() && TryEnterSyncBarrier(InSyncStep))
@@ -66,6 +70,8 @@ bool FTextureShareCoreObject::FrameSync(const ETextureShareSyncStep InSyncStep)
 
 bool FTextureShareCoreObject::FrameSync_RenderThread(const ETextureShareSyncStep InSyncStep)
 {
+	TRACE_CPUPROFILER_EVENT_SCOPE(TextureShareCore::FrameSync_RenderThread);
+
 	UE_TS_LOG(LogTextureShareCoreObjectSync, Log, TEXT(">>>> %s:FrameSync_RenderThread(%s)"), *GetName(), GetTEXT(InSyncStep));
 
 	if (IsFrameSyncActive() && TryEnterSyncBarrier(InSyncStep))
@@ -148,6 +154,12 @@ bool FTextureShareCoreObject::BeginFrameSync()
 	{
 		SetCurrentSyncStep(ETextureShareSyncStep::FrameBegin);
 		SetFrameSyncState(ETextureShareCoreInterprocessObjectFrameSyncState::FrameBegin);
+
+		if (ITextureShareCoreCallbacks::Get().OnTextureShareCoreBeginFrameSync().IsBound())
+		{
+			ITextureShareCoreCallbacks::Get().OnTextureShareCoreBeginFrameSync().Broadcast(*this);
+		}
+
 		return true;
 	}
 
@@ -161,6 +173,12 @@ bool FTextureShareCoreObject::EndFrameSync()
 		SetCurrentSyncStep(ETextureShareSyncStep::FrameEnd);
 		SetFrameSyncState(ETextureShareCoreInterprocessObjectFrameSyncState::FrameEnd);
 
+
+		if (ITextureShareCoreCallbacks::Get().OnTextureShareCoreEndFrameSync().IsBound())
+		{
+			ITextureShareCoreCallbacks::Get().OnTextureShareCoreEndFrameSync().Broadcast(*this);
+		}
+
 		return true;
 	}
 
@@ -195,6 +213,12 @@ bool FTextureShareCoreObject::BeginFrameSync_RenderThread()
 
 		SetCurrentSyncStep(ETextureShareSyncStep::FrameProxyBegin);
 		SetFrameSyncState(ETextureShareCoreInterprocessObjectFrameSyncState::FrameProxyBegin);
+
+		if (ITextureShareCoreCallbacks::Get().OnTextureShareCoreBeginFrameSync_RenderThread().IsBound())
+		{
+			ITextureShareCoreCallbacks::Get().OnTextureShareCoreBeginFrameSync_RenderThread().Broadcast(*this);
+		}
+
 		return true;
 	}
 
@@ -217,6 +241,11 @@ bool FTextureShareCoreObject::EndFrameSync_RenderThread()
 
 		SetCurrentSyncStep(ETextureShareSyncStep::FrameProxyEnd);
 		SetFrameSyncState(ETextureShareCoreInterprocessObjectFrameSyncState::FrameProxyEnd);
+
+		if (ITextureShareCoreCallbacks::Get().OnTextureShareCoreEndFrameSync_RenderThread().IsBound())
+		{
+			ITextureShareCoreCallbacks::Get().OnTextureShareCoreEndFrameSync_RenderThread().Broadcast(*this);
+		}
 	}
 
 	return bResult;
@@ -272,6 +301,11 @@ bool FTextureShareCoreObject::DoFrameSync(const ETextureShareSyncStep InSyncStep
 		// Write local data
 		SendFrameData();
 
+		if (ITextureShareCoreCallbacks::Get().OnTextureShareCoreFrameSync().IsBound())
+		{
+			ITextureShareCoreCallbacks::Get().OnTextureShareCoreFrameSync().Broadcast(*this, InSyncStep);
+		}
+
 		// add barrier here
 		if (SyncBarrierPass(InSyncStep, ETextureShareSyncPass::Enter))
 		{
@@ -308,6 +342,11 @@ bool FTextureShareCoreObject::DoFrameSync_RenderThread(const ETextureShareSyncSt
 		// Write local data
 		SendFrameProxyData_RenderThread();
 
+		if (ITextureShareCoreCallbacks::Get().OnTextureShareCoreFrameSync_RenderThread().IsBound())
+		{
+			ITextureShareCoreCallbacks::Get().OnTextureShareCoreFrameSync_RenderThread().Broadcast(*this, InSyncStep);
+		}
+
 		// add barrier here
 		if (SyncBarrierPass(InSyncStep, ETextureShareSyncPass::Enter))
 		{
