TextureShare: Added DepthBuffer format support

diff --git a/Engine/Plugins/VirtualProduction/TextureShare/Source/TextureShare/Private/Resources/TextureShareResourceSettings.cpp b/Engine/Plugins/VirtualProduction/TextureShare/Source/TextureShare/Private/Resources/TextureShareResourceSettings.cpp
--- a/Engine/Plugins/VirtualProduction/TextureShare/Source/TextureShare/Private/Resources/TextureShareResourceSettings.cpp
+++ b/Engine/Plugins/VirtualProduction/TextureShare/Source/TextureShare/Private/Resources/TextureShareResourceSettings.cpp
@@ -80,6 +80,7 @@ FTextureShareResourceSettings::FTextureShareResourceSettings(const FTextureShare
 		}
 	}
 
+	// Not all types of texture formats can be used for sharing.
 	// Use only formats supported by shared D3D
 	switch (Format)
 	{
@@ -87,18 +88,51 @@ FTextureShareResourceSettings::FTextureShareResourceSettings(const FTextureShare
 	case PF_R8G8B8A8:
 	case PF_A8R8G8B8:
 	case PF_B8G8R8A8:
-	case PF_A32B32G32R32F:
 	case PF_FloatRGBA:
-	case PF_FloatRGB:
-		// Define more supported formats
-		// This requires a lot of tests.
-		// Todo: improve this in the future (It also requires symmetrical changes to the TextureShareResourceUtils.cpp file.)
-		//EPixelFormat::PF_D24,
-		//EPixelFormat::PF_R32_FLOAT,
+		// These formats are supported without conversion
+		break;
+
+	// 16-bit depth
+	case PF_A16B16G16R16:
+	case PF_R16G16B16A16_UNORM:
+	case PF_R16G16B16A16_SNORM:
+	case PF_R16G16B16A16_UINT:
+	case PF_R16G16B16A16_SINT:
+	case PF_G16R16:
+	case PF_G16R16F:
+	case PF_G16R16F_FILTER:
+	case PF_G16R16_SNORM:
+	case PF_R16G16_UINT:
+	case PF_R16_UINT:
+	case PF_R16_SINT:
+	case PF_G16:
+		// These formats are not supported without conversion and must be converted to 16-bit ARGB to prevent data loss.
+		Format = PF_A16B16G16R16;
+		break;
+
+	// 32-bit depth
+	case PF_DepthStencil:
+	case PF_ShadowDepth:
+	case PF_D24:
+	case PF_X24_G8:
+	case PF_G32R32F:
+	case PF_R32G32_UINT:
+	case PF_R32_UINT:
+	case PF_R32_SINT:
+	case PF_R32_FLOAT:
+	case PF_R32G32B32_UINT:
+	case PF_R32G32B32_SINT:
+	case PF_R32G32B32F:
+	case PF_R32G32B32A32_UINT:
+	case PF_A32B32G32R32F:
+		// These formats are not supported without conversion and must be converted to 32-bit ARGB to prevent data loss.
+		Format = EPixelFormat::PF_A32B32G32R32F;
 		break;
 
 	default:
-		// for all unsupported formats use default
+		// By default, any format is converted to PF_B8G8R8A8.
+		// This can result in loss of data for the high-bit input format.
+		// Todo: Add more conversion rules for formats that don't fit the default format type.
 		Format = EPixelFormat::PF_B8G8R8A8;
 		break;
 	}
