Fixing include paths to make them less hardcoded.

diff --git a/Engine/Plugins/VirtualProduction/TextureShare/Source/TextureShare/TextureShare.Build.cs b/Engine/Plugins/VirtualProduction/TextureShare/Source/TextureShare/TextureShare.Build.cs
--- a/Engine/Plugins/VirtualProduction/TextureShare/Source/TextureShare/TextureShare.Build.cs
+++ b/Engine/Plugins/VirtualProduction/TextureShare/Source/TextureShare/TextureShare.Build.cs
@@ -10,8 +10,8 @@ public class TextureShare : ModuleRules
 	{
 		PrivateIncludePaths.AddRange(
 			new string[] {
-				Path.Combine(EngineDirectory,"Source/Runtime/Renderer/Private"),
-				Path.Combine(EngineDirectory,"Plugins/VirtualProduction/TextureShare/Source/TextureShareCore/Private"),
+				Path.Combine(GetModuleDirectory("Renderer"), "Private"),
+				Path.Combine(GetModuleDirectory("TextureShareCore"), "Private"),
 			}
 		);
 
