Adding API accessors for SceneTexture planes.

diff --git a/Engine/Plugins/VirtualProduction/TextureShare/Source/TextureShare/Private/Game/ViewExtension/TextureShareSceneViewExtension.cpp b/Engine/Plugins/VirtualProduction/TextureShare/Source/TextureShare/Private/Game/ViewExtension/TextureShareSceneViewExtension.cpp
--- a/Engine/Plugins/VirtualProduction/TextureShare/Source/TextureShare/Private/Game/ViewExtension/TextureShareSceneViewExtension.cpp
+++ b/Engine/Plugins/VirtualProduction/TextureShare/Source/TextureShare/Private/Game/ViewExtension/TextureShareSceneViewExtension.cpp
@@ -8,13 +8,13 @@
 #include "Module/TextureShareLog.h"
 #include "Misc/TextureShareStrings.h"
 
+#include "RenderGraphBuilder.h"
+#include "RenderGraphUtils.h"
 #include "RHI.h"
 #include "RHICommandList.h"
 #include "RHIResources.h"
+#include "SceneRenderTargetParameters.h"
 #include "SceneView.h"
-#include "SceneTextures.h"
-
-#include "PostProcess/SceneRenderTargets.h"
 
 //////////////////////////////////////////////////////////////////////////////////////////////
 namespace TextureShareSceneViewExtensionHelpers
@@ -133,24 +133,26 @@ void FTextureShareSceneViewExtension::GetSceneViewData_RenderThread(const FTextu
 
 void FTextureShareSceneViewExtension::ShareSceneViewColors_RenderThread(FRDGBuilder& GraphBuilder, const FSceneTextures& SceneTextures, const FTextureShareSceneView& InView)
 {
-	const auto AddShareTexturePass = [&](const TCHAR* InTextureName, const FRDGTextureRef& InTextureRef)
+	const auto AddShareTexturePass = [&](const TCHAR* InTextureName, ESceneTexture InSceneTexture)
 	{
+		const FRDGTextureRef InTextureRef = GetSceneTexture(SceneTextures, InSceneTexture);
+
 		// Send resource
 		ObjectProxy->ShareResource_RenderThread(GraphBuilder, FTextureShareCoreResourceDesc(InTextureName, InView.ViewInfo.ViewDesc, ETextureShareTextureOp::Read),
 			InTextureRef, InView.GPUIndex, &InView.UnconstrainedViewRect);
 	};
 
-	AddShareTexturePass(TextureShareStrings::SceneTextures::SceneColor, SceneTextures.Color.Resolve);
+	AddShareTexturePass(TextureShareStrings::SceneTextures::SceneColor, ESceneTexture::Color);
 
-	AddShareTexturePass(TextureShareStrings::SceneTextures::SceneDepth, SceneTextures.Depth.Resolve);
-	AddShareTexturePass(TextureShareStrings::SceneTextures::SmallDepthZ, SceneTextures.SmallDepth);
+	AddShareTexturePass(TextureShareStrings::SceneTextures::SceneDepth, ESceneTexture::Depth);
+	AddShareTexturePass(TextureShareStrings::SceneTextures::SmallDepthZ, ESceneTexture::SmallDepth);
 
-	AddShareTexturePass(TextureShareStrings::SceneTextures::GBufferA, SceneTextures.GBufferA);
-	AddShareTexturePass(TextureShareStrings::SceneTextures::GBufferB, SceneTextures.GBufferB);
-	AddShareTexturePass(TextureShareStrings::SceneTextures::GBufferC, SceneTextures.GBufferC);
-	AddShareTexturePass(TextureShareStrings::SceneTextures::GBufferD, SceneTextures.GBufferD);
-	AddShareTexturePass(TextureShareStrings::SceneTextures::GBufferE, SceneTextures.GBufferE);
-	AddShareTexturePass(TextureShareStrings::SceneTextures::GBufferF, SceneTextures.GBufferF);
+	AddShareTexturePass(TextureShareStrings::SceneTextures::GBufferA, ESceneTexture::GBufferA);
+	AddShareTexturePass(TextureShareStrings::SceneTextures::GBufferB, ESceneTexture::GBufferB);
+	AddShareTexturePass(TextureShareStrings::SceneTextures::GBufferC, ESceneTexture::GBufferC);
+	AddShareTexturePass(TextureShareStrings::SceneTextures::GBufferD, ESceneTexture::GBufferD);
+	AddShareTexturePass(TextureShareStrings::SceneTextures::GBufferE, ESceneTexture::GBufferE);
+	AddShareTexturePass(TextureShareStrings::SceneTextures::GBufferF, ESceneTexture::GBufferF);
 }
 
 //////////////////////////////////////////////////////////////////////////////////////////////
diff --git a/Engine/Source/Runtime/Renderer/Private/SceneTextures.cpp b/Engine/Source/Runtime/Renderer/Private/SceneTextures.cpp
--- a/Engine/Source/Runtime/Renderer/Private/SceneTextures.cpp
+++ b/Engine/Source/Runtime/Renderer/Private/SceneTextures.cpp
@@ -767,6 +767,28 @@ void QueueSceneTextureExtractions(FRDGBuilder& GraphBuilder, const FSceneTexture
 	return GSceneTextureExtracts.QueueExtractions(GraphBuilder, SceneTextures);
 }
 
+FRDGTextureRef GetSceneTexture(const FSceneTextures& SceneTextures, ESceneTexture InSceneTexture)
+{
+	switch (InSceneTexture)
+	{
+	case ESceneTexture::Color:          return SceneTextures.Color.Resolve;
+	case ESceneTexture::Depth:          return SceneTextures.Depth.Resolve;
+	case ESceneTexture::SmallDepth:     return SceneTextures.SmallDepth;
+	case ESceneTexture::Velocity:       return SceneTextures.Velocity;
+	case ESceneTexture::GBufferA:       return SceneTextures.GBufferA;
+	case ESceneTexture::GBufferB:       return SceneTextures.GBufferB;
+	case ESceneTexture::GBufferC:       return SceneTextures.GBufferC;
+	case ESceneTexture::GBufferD:       return SceneTextures.GBufferD;
+	case ESceneTexture::GBufferE:       return SceneTextures.GBufferE;
+	case ESceneTexture::GBufferF:       return SceneTextures.GBufferF;
+	case ESceneTexture::SSAO:           return SceneTextures.ScreenSpaceAO;
+	case ESceneTexture::CustomDepth:	return SceneTextures.CustomDepth.Depth;
+	default:
+		checkNoEntry();
+		return nullptr;
+	}
+}
+
 void SetupSceneTextureUniformParameters(
 	FRDGBuilder& GraphBuilder,
 	const FSceneTextures* SceneTextures,
diff --git a/Engine/Source/Runtime/Renderer/Public/SceneRenderTargetParameters.h b/Engine/Source/Runtime/Renderer/Public/SceneRenderTargetParameters.h
--- a/Engine/Source/Runtime/Renderer/Public/SceneRenderTargetParameters.h
+++ b/Engine/Source/Runtime/Renderer/Public/SceneRenderTargetParameters.h
@@ -11,6 +11,24 @@
 class FRDGBuilder;
 struct FSceneTextures;
 
+enum class ESceneTexture
+{
+	Color,
+	Depth,
+	SmallDepth,
+	Velocity,
+	GBufferA,
+	GBufferB,
+	GBufferC,
+	GBufferD,
+	GBufferE,
+	GBufferF,
+	SSAO,
+	CustomDepth,
+};
+
+RENDERER_API FRDGTextureRef GetSceneTexture(const FSceneTextures& SceneTextures, ESceneTexture InSceneTexture);
+
 enum class ESceneTextureSetupMode : uint32
 {
 	None			= 0,
