Removing unnecessary include.

diff --git a/Engine/Plugins/VirtualProduction/TextureShare/Source/TextureShare/Private/Resources/TextureShareResourcesProxy.cpp b/Engine/Plugins/VirtualProduction/TextureShare/Source/TextureShare/Private/Resources/TextureShareResourcesProxy.cpp
--- a/Engine/Plugins/VirtualProduction/TextureShare/Source/TextureShare/Private/Resources/TextureShareResourcesProxy.cpp
+++ b/Engine/Plugins/VirtualProduction/TextureShare/Source/TextureShare/Private/Resources/TextureShareResourcesProxy.cpp
@@ -18,8 +18,6 @@
 #include "RHIStaticStates.h"
 
 #include "RenderingThread.h"
-#include "RendererPrivate.h"
-
 #include "RenderResource.h"
 #include "CommonRenderResources.h"
 #include "PixelShaderUtils.h"
