TextureShare: Fixed compiler issues for string conversion local variables.

diff --git a/Engine/Plugins/VirtualProduction/TextureShare/Source/TextureShareCore/Private/Core/Serialize/ITextureShareCoreSerializeStream.h b/Engine/Plugins/VirtualProduction/TextureShare/Source/TextureShareCore/Private/Core/Serialize/ITextureShareCoreSerializeStream.h
--- a/Engine/Plugins/VirtualProduction/TextureShare/Source/TextureShareCore/Private/Core/Serialize/ITextureShareCoreSerializeStream.h
+++ b/Engine/Plugins/VirtualProduction/TextureShare/Source/TextureShareCore/Private/Core/Serialize/ITextureShareCoreSerializeStream.h
@@ -63,11 +63,13 @@ public:
 		int32 StrLen = InOutData.Len();
 		*this << StrLen;
 
+		// Always serialize string as wchar_t (An external application from the SDK side can easily use this string type without conversion)
 		if (StrLen > 0)
 		{
 			if (IsWriteStream())
 			{
-				SerializeData(FTCHARToWChar(*InOutData).Get(), sizeof(wchar_t) * StrLen);
+				FTCHARToWChar WCharString(*InOutData, StrLen);
+				SerializeData(WCharString.Get(), sizeof(wchar_t) * StrLen);
 			}
 			else
 			{
@@ -75,7 +77,8 @@ public:
 				WCharBuffer.AddZeroed(StrLen + 1);
 				SerializeData(WCharBuffer.GetData(), sizeof(wchar_t) * StrLen);
 
-				InOutData = FWCharToTCHAR(WCharBuffer.GetData()).Get();
+				FWCharToTCHAR TCHARString(WCharBuffer.GetData());
+				InOutData = TCHARString.Get();
 			}
 		}
 
diff --git a/Engine/Plugins/VirtualProduction/TextureShare/Source/TextureShareCore/Private/IPC/Containers/TextureShareCoreInterprocessContainers.cpp b/Engine/Plugins/VirtualProduction/TextureShare/Source/TextureShareCore/Private/IPC/Containers/TextureShareCoreInterprocessContainers.cpp
--- a/Engine/Plugins/VirtualProduction/TextureShare/Source/TextureShareCore/Private/IPC/Containers/TextureShareCoreInterprocessContainers.cpp
+++ b/Engine/Plugins/VirtualProduction/TextureShare/Source/TextureShareCore/Private/IPC/Containers/TextureShareCoreInterprocessContainers.cpp
@@ -28,11 +28,9 @@ void FTextureShareCoreSMD5Hash::Initialize(const FString& InText)
 {
 	Empty();
 
-	// Create MD5 Hash
-	uint64 BlobInputLen = FCString::Strlen(*InText);
-
+	// Generate MD5 Hash from TCHAR as binary data
 	FMD5 Md5Gen;
-	Md5Gen.Update((unsigned char*)TCHAR_TO_ANSI(*InText), BlobInputLen);
+	Md5Gen.Update((const uint8*)GetData(InText), GetNum(InText) * sizeof(TCHAR));
 	Md5Gen.Final(MD5Digest);
 }
 
@@ -68,20 +66,12 @@ void FTextureShareCoreStringHash::Initialize(const FString& InText)
 
 	Hash.Initialize(InText);
 
-	// Copy string
 	FTCHARToWChar WCharString(*InText);
-	const wchar_t* WCharValue = WCharString.Get();
-	if (WCharValue && WCharValue[0])
 	{
-		for (int32 CharIt = 0; CharIt < (MaxStringLength-1); CharIt++)
-		{
-			String[CharIt] = WCharValue[CharIt];
+		const int32 StrLength = FMath::Min(WCharString.Length(), MaxStringLength - 1);
 
-			// EndOfString
-			if (WCharValue[CharIt] == 0)
-			{
-				break;
-			}
-		}
+		// Reset and copy wchar string
+		FPlatformMemory::Memset(&String[0], 0, sizeof(String));
+		FPlatformMemory::Memcpy(&String[0], WCharString.Get(), sizeof(wchar_t) * StrLength);
 	}
 }
