TextureShare: Fix merge issue in build file.

#rnx

diff --git a/Engine/Plugins/VirtualProduction/TextureShare/Source/TextureShare/TextureShare.Build.cs b/Engine/Plugins/VirtualProduction/TextureShare/Source/TextureShare/TextureShare.Build.cs
--- a/Engine/Plugins/VirtualProduction/TextureShare/Source/TextureShare/TextureShare.Build.cs
+++ b/Engine/Plugins/VirtualProduction/TextureShare/Source/TextureShare/TextureShare.Build.cs
@@ -12,9 +12,8 @@ public class TextureShare : ModuleRules
 		// be removed once we move the nD rendering pipeline to RDG (TextureShareSceneViewExtension.h).
 		string EngineDir = Path.GetFullPath(Target.RelativeEnginePath);
 		PrivateIncludePaths.Add(Path.Combine(EngineDir, "Source", "Runtime", "Renderer", "Private"));
-
-		// Internal dependency (debug log purpose)
-		PrivateIncludePaths.Add(Path.Combine(EngineDir, "Plugins", "VirtualProduction", "TextureShare", "Source", "TextureShareCore", "Private"));
+			}
+		);
 
 		PublicDependencyModuleNames.AddRange(
 			new string[]
