TextureShare: Fix merge issue.

This is a temp fix for compilation purposes.

#rnx

diff --git a/Engine/Plugins/VirtualProduction/TextureShare/Source/TextureShare/Private/Resources/TextureShareResourcesProxy.cpp b/Engine/Plugins/VirtualProduction/TextureShare/Source/TextureShare/Private/Resources/TextureShareResourcesProxy.cpp
--- a/Engine/Plugins/VirtualProduction/TextureShare/Source/TextureShare/Private/Resources/TextureShareResourcesProxy.cpp
+++ b/Engine/Plugins/VirtualProduction/TextureShare/Source/TextureShare/Private/Resources/TextureShareResourcesProxy.cpp
@@ -105,7 +105,7 @@ namespace UE::TextureShare::ResourcesProxy
 			const bool bSameSize = (SrcRect.Size() == DstRect.Size());
 			FRHISamplerState* PixelSampler = bSameSize ? TStaticSamplerState<SF_Point>::GetRHI() : TStaticSamplerState<SF_Bilinear>::GetRHI();
 
-			SetShaderParametersLegacyPS(RHICmdList, PixelShader, PixelSampler, SrcTexture);
+			//SetShaderParametersLegacyPS(RHICmdList, PixelShader, PixelSampler, SrcTexture);
 
 			// Set up vertex uniform parameters for scaling and biasing the rectangle.
 			// Note: Use DrawRectangle in the vertex shader to calculate the correct vertex position and uv.
