TextureShare: Fix bad merge of a bad merge fix.

#rnx

diff --git a/Engine/Plugins/VirtualProduction/TextureShare/Source/TextureShare/Private/Resources/TextureShareResourcesProxy.cpp b/Engine/Plugins/VirtualProduction/TextureShare/Source/TextureShare/Private/Resources/TextureShareResourcesProxy.cpp
--- a/Engine/Plugins/VirtualProduction/TextureShare/Source/TextureShare/Private/Resources/TextureShareResourcesProxy.cpp
+++ b/Engine/Plugins/VirtualProduction/TextureShare/Source/TextureShare/Private/Resources/TextureShareResourcesProxy.cpp
@@ -105,7 +105,7 @@ namespace UE::TextureShare::ResourcesProxy
 			const bool bSameSize = (SrcRect.Size() == DstRect.Size());
 			FRHISamplerState* PixelSampler = bSameSize ? TStaticSamplerState<SF_Point>::GetRHI() : TStaticSamplerState<SF_Bilinear>::GetRHI();
 
-			//SetShaderParametersLegacyPS(RHICmdList, PixelShader, PixelSampler, SrcTexture);
+			SetShaderParametersLegacyPS(RHICmdList, PixelShader, PixelSampler, SrcTexture);
 
 			// Set up vertex uniform parameters for scaling and biasing the rectangle.
 			// Note: Use DrawRectangle in the vertex shader to calculate the correct vertex position and uv.
diff --git a/Engine/Plugins/VirtualProduction/TextureShare/Source/TextureShare/TextureShare.Build.cs b/Engine/Plugins/VirtualProduction/TextureShare/Source/TextureShare/TextureShare.Build.cs
--- a/Engine/Plugins/VirtualProduction/TextureShare/Source/TextureShare/TextureShare.Build.cs
+++ b/Engine/Plugins/VirtualProduction/TextureShare/Source/TextureShare/TextureShare.Build.cs
@@ -8,12 +8,10 @@ public class TextureShare : ModuleRules
 {
 	public TextureShare(ReadOnlyTargetRules ROTargetRules) : base(ROTargetRules)
 	{
-		// [temporary] We need this to be able to use some private data types. This should
-		// be removed once we move the nD rendering pipeline to RDG (TextureShareSceneViewExtension.h).
-		string EngineDir = Path.GetFullPath(Target.RelativeEnginePath);
-		PrivateIncludePaths.Add(Path.Combine(EngineDir, "Source", "Runtime", "Renderer", "Private"));
-			}
-		);
+		PrivateIncludePaths.AddRange(new string[] {
+			Path.Combine(GetModuleDirectory("Renderer"), "Private"),
+			Path.Combine(GetModuleDirectory("TextureShareCore"), "Private"),
+		});
 
 		PublicDependencyModuleNames.AddRange(
 			new string[]
