TextureShare:
* added missing headers (after CL 24666180)
* simplification of nested namespaces
* updated SDK files

diff --git a/Engine/Extras/VirtualProduction/TextureShare/TextureShareSDK/Source/Containers/TextureShareCoreContainers_ObjectDesc.h b/Engine/Extras/VirtualProduction/TextureShare/TextureShareSDK/Source/Containers/TextureShareCoreContainers_ObjectDesc.h
--- a/Engine/Extras/VirtualProduction/TextureShare/TextureShareSDK/Source/Containers/TextureShareCoreContainers_ObjectDesc.h
+++ b/Engine/Extras/VirtualProduction/TextureShare/TextureShareSDK/Source/Containers/TextureShareCoreContainers_ObjectDesc.h
@@ -2,6 +2,7 @@
 
 #pragma once
 #include "Serialize/TextureShareCoreSerialize.h"
+#include "Containers/TextureShareCoreContainers_SyncSettings.h"
 
 /**
  * Process descriptor
@@ -38,11 +39,11 @@ public:
 struct FTextureShareCoreObjectSyncState
 	: public ITextureShareSerialize
 {
-	ETextureShareSyncStep  Step;
-	ETextureShareSyncState State;
+	ETextureShareSyncStep  Step = ETextureShareSyncStep::Undefined;
+	ETextureShareSyncState State = ETextureShareSyncState::Undefined;
 
-	ETextureShareSyncStep NextStep;
-	ETextureShareSyncStep PrevStep;
+	ETextureShareSyncStep NextStep = ETextureShareSyncStep::Undefined;
+	ETextureShareSyncStep PrevStep = ETextureShareSyncStep::Undefined;
 
 public:
 	virtual ~FTextureShareCoreObjectSyncState() = default;
@@ -64,7 +65,7 @@ struct FTextureShareCoreObjectSync
 
 	// Bit-storage for sync steps used by this process
 	// 64 bits. Max steps=64
-	uint64 SyncStepSettings;
+	uint64 SyncStepSettings = 0;
 
 	// Last access time from this object
 	uint64 LastAccessTime = 0;
@@ -84,7 +85,7 @@ public:
 			return true;
 		}
 
-		const int8 BitIndex = (int8)InStep;
+		const uint8 BitIndex = (uint8)InStep;
 
 		if (BitIndex >= 0 && BitIndex < 64)
 		{
@@ -93,6 +94,20 @@ public:
 
 		return false;
 	}
+
+	void SetSyncStepSettings(const FTextureShareCoreSyncSettings& InSyncSettings)
+	{
+		// Update bit-storage for used steps
+		SyncStepSettings = 0;
+		for (const ETextureShareSyncStep& StepIt : InSyncSettings.FrameSyncSettings.Steps)
+		{
+			const uint8 BitIndex = (uint8)StepIt;
+			if (BitIndex >= 0 && BitIndex < 64)
+			{
+				SyncStepSettings |= (1ULL << BitIndex);
+			}
+		}
+	}
 };
 
 /**
diff --git a/Engine/Extras/VirtualProduction/TextureShare/TextureShareSDK/Source/Containers/TextureShareCoreContainers_SyncSettings.h b/Engine/Extras/VirtualProduction/TextureShare/TextureShareSDK/Source/Containers/TextureShareCoreContainers_SyncSettings.h
--- a/Engine/Extras/VirtualProduction/TextureShare/TextureShareSDK/Source/Containers/TextureShareCoreContainers_SyncSettings.h
+++ b/Engine/Extras/VirtualProduction/TextureShare/TextureShareSDK/Source/Containers/TextureShareCoreContainers_SyncSettings.h
@@ -6,7 +6,7 @@
 /**
  * Connection state
  */
-enum class ETextureShareCoreFrameConnectionsState : int8
+enum class ETextureShareCoreFrameConnectionsState : uint8
 {
 	SkipFrame = 0,
 	Wait,
@@ -37,6 +37,13 @@ public:
 		return Stream << MinValue << AllowedProcessNames << BannedProcessNames;
 	}
 
+	bool operator==(const FTextureShareCoreFrameConnectionsSettings& In) const
+	{
+		return MinValue == In.MinValue
+			&& AllowedProcessNames.EqualsFunc(In.AllowedProcessNames)
+			&& BannedProcessNames.EqualsFunc(In.BannedProcessNames);
+	}
+
 public:
 	ETextureShareCoreFrameConnectionsState GetConnectionsState(const int32 InReadyToConnectObjectsCount, const int32 InFrameConnectionsCount) const
 	{
@@ -78,16 +85,6 @@ public:
 	}
 
 public:
-	bool Contains(const ETextureShareSyncStep InSyncStep) const
-	{
-		return Steps.Find(InSyncStep) != INDEX_NONE;
-	}
-
-	void Append(const TArray<ETextureShareSyncStep>& InSyncSteps)
-	{
-		Steps.AppendSorted(InSyncSteps);
-	}
-
 	bool EqualsFunc(const FTextureShareCoreFrameSyncSettings& In) const
 	{
 		return Steps.EqualsFunc(In.Steps);
@@ -135,6 +132,17 @@ public:
 			<< ProcessLostStatusTimeOut
 			<< ThreadMutexTimeout;
 	}
+
+	bool operator==(const FTextureShareCoreTimeOutSettings& In) const
+	{
+		return MemoryMutexTimeout == In.MemoryMutexTimeout
+			&& FrameBeginTimeOut == In.FrameBeginTimeOut
+			&& FrameSyncTimeOut == In.FrameSyncTimeOut
+			&& FrameBeginTimeOutSplit == In.FrameBeginTimeOutSplit
+			&& FrameSyncTimeOutSplit == In.FrameSyncTimeOutSplit
+			&& ProcessLostStatusTimeOut == In.ProcessLostStatusTimeOut
+			&& ThreadMutexTimeout == In.ThreadMutexTimeout;
+	}
 };
 
 /**
@@ -161,4 +169,11 @@ public:
 	{
 		return Stream << FrameConnectionSettings << TimeoutSettings << FrameSyncSettings;
 	}
+
+	bool operator==(const FTextureShareCoreSyncSettings& In) const
+	{
+		return FrameConnectionSettings == (In.FrameConnectionSettings)
+			&& TimeoutSettings == In.TimeoutSettings
+			&& FrameSyncSettings.EqualsFunc(In.FrameSyncSettings);
+	}
 };
diff --git a/Engine/Extras/VirtualProduction/TextureShare/TextureShareSDK/Source/Misc/TextureShareCoreStrings.h b/Engine/Extras/VirtualProduction/TextureShare/TextureShareSDK/Source/Misc/TextureShareCoreStrings.h
--- a/Engine/Extras/VirtualProduction/TextureShare/TextureShareSDK/Source/Misc/TextureShareCoreStrings.h
+++ b/Engine/Extras/VirtualProduction/TextureShare/TextureShareSDK/Source/Misc/TextureShareCoreStrings.h
@@ -10,23 +10,20 @@
 /**
  * Default names for TextureShareCore module
  */
-namespace TextureShareCoreStrings
+namespace UE::TextureShareCoreStrings
 {
-	namespace Default
-	{
-		static constexpr auto ShareName = TEXT("DefaultShareName");
+	static constexpr auto DefaultShareName = TEXT("DefaultShareName");
+
+	static constexpr auto DefaultViewId = TEXT("DefaultView");
 
-		namespace ProcessName
-		{
+	namespace DefaultProcessName
+	{
 #if defined TEXTURESHARE_PROJECT_NAME
-			// Use TextureShare project name as process name
-			static constexpr auto SDK = TEXT(TEXTURESHARE_PROJECT_NAME);
+		// Use TextureShare project name as process name
+		static constexpr auto SDK = TEXT(TEXTURESHARE_PROJECT_NAME);
 #else
-			static constexpr auto SDK = TEXT("TextureShareSDK");
+		static constexpr auto SDK = TEXT("TextureShareSDK");
 #endif
-			static constexpr auto UE = TEXT("UnrealEngine");
-		}
-
-		static constexpr auto ViewId = TEXT("DefaultView");
+		static constexpr auto UE = TEXT("UnrealEngine");
 	}
 };
diff --git a/Engine/Extras/VirtualProduction/TextureShare/TextureShareSDK/Source/Misc/TextureShareDisplayClusterStrings.h b/Engine/Extras/VirtualProduction/TextureShare/TextureShareSDK/Source/Misc/TextureShareDisplayClusterStrings.h
--- a/Engine/Extras/VirtualProduction/TextureShare/TextureShareSDK/Source/Misc/TextureShareDisplayClusterStrings.h
+++ b/Engine/Extras/VirtualProduction/TextureShare/TextureShareSDK/Source/Misc/TextureShareDisplayClusterStrings.h
@@ -10,44 +10,38 @@
 /**
  * Constant names for nDisplay texture resources
  */
-namespace UE
+namespace UE::TextureShare::DisplayClusterStrings
 {
-	namespace TextureShare
+	// Default ShareName for nDisplay integration object
+	static constexpr auto DefaultShareName = TEXT("nDisplay");
+
+	namespace Viewport
+	{
+		// FinalColor, but can be overrided
+		static constexpr auto FinalColor = TEXT("ViewportFinalColor");
+
+		// internal resolved shader resources, used as warp&blend source
+		static constexpr auto Input = TEXT("ViewportInput");
+		static constexpr auto Mips = TEXT("ViewportMips");
+
+		// After warp viewport (before output remap)
+		static constexpr auto Warped = TEXT("ViewportWarpBlend");
+	}
+
+	namespace Output
+	{
+		// access to nDisplay frame backbuffer
+		static constexpr auto Backbuffer = TEXT("FrameBackbuffer");
+		static constexpr auto BackbufferTemp = TEXT("FrameBackbufferTemp");
+	}
+
+	namespace Postprocess
+	{
+		static constexpr auto TextureShare = TEXT("TextureShare");
+	}
+
+	namespace Projection
 	{
-		namespace DisplayClusterStrings
-		{
-			// Default ShareName for nDisplay integration object
-			static constexpr auto DefaultShareName = TEXT("nDisplay");
-
-			namespace Viewport
-			{
-				// FinalColor, but can be overrided
-				static constexpr auto FinalColor = TEXT("ViewportFinalColor");
-
-				// internal resolved shader resources, used as warp&blend source
-				static constexpr auto Input = TEXT("ViewportInput");
-				static constexpr auto Mips = TEXT("ViewportMips");
-
-				// After warp viewport (before output remap)
-				static constexpr auto Warped = TEXT("ViewportWarpBlend");
-			}
-
-			namespace Output
-			{
-				// access to nDisplay frame backbuffer
-				static constexpr auto Backbuffer = TEXT("FrameBackbuffer");
-				static constexpr auto BackbufferTemp = TEXT("FrameBackbufferTemp");
-			}
-
-			namespace Postprocess
-			{
-				static constexpr auto TextureShare = TEXT("TextureShare");
-			}
-
-			namespace Projection
-			{
-				static constexpr auto TextureShare = TEXT("textureshare");
-			}
-		}
+		static constexpr auto TextureShare = TEXT("textureshare");
 	}
 };
diff --git a/Engine/Extras/VirtualProduction/TextureShare/TextureShareSDK/Source/Misc/TextureShareStrings.h b/Engine/Extras/VirtualProduction/TextureShare/TextureShareSDK/Source/Misc/TextureShareStrings.h
--- a/Engine/Extras/VirtualProduction/TextureShare/TextureShareSDK/Source/Misc/TextureShareStrings.h
+++ b/Engine/Extras/VirtualProduction/TextureShare/TextureShareSDK/Source/Misc/TextureShareStrings.h
@@ -10,23 +10,20 @@
 /**
  * Default names for TextureShare module
  */
-namespace TextureShareStrings
+namespace UE::TextureShareStrings::SceneTextures
 {
-	namespace SceneTextures
-	{
-		// Read-only scene resources
-		static constexpr auto SceneColor = TEXT("SceneColor");
-		static constexpr auto SceneDepth = TEXT("SceneDepth");
-		static constexpr auto SmallDepthZ = TEXT("SmallDepthZ");
-		static constexpr auto GBufferA = TEXT("GBufferA");
-		static constexpr auto GBufferB = TEXT("GBufferB");
-		static constexpr auto GBufferC = TEXT("GBufferC");
-		static constexpr auto GBufferD = TEXT("GBufferD");
-		static constexpr auto GBufferE = TEXT("GBufferE");
-		static constexpr auto GBufferF = TEXT("GBufferF");
+	// Read-only scene resources
+	static constexpr auto SceneColor = TEXT("SceneColor");
+	static constexpr auto SceneDepth = TEXT("SceneDepth");
+	static constexpr auto SmallDepthZ = TEXT("SmallDepthZ");
+	static constexpr auto GBufferA = TEXT("GBufferA");
+	static constexpr auto GBufferB = TEXT("GBufferB");
+	static constexpr auto GBufferC = TEXT("GBufferC");
+	static constexpr auto GBufferD = TEXT("GBufferD");
+	static constexpr auto GBufferE = TEXT("GBufferE");
+	static constexpr auto GBufferF = TEXT("GBufferF");
 
-		// Read-write RTTs
-		static constexpr auto FinalColor = TEXT("FinalColor");
-		static constexpr auto Backbuffer = TEXT("Backbuffer");
-	}
+	// Read-write RTTs
+	static constexpr auto FinalColor = TEXT("FinalColor");
+	static constexpr auto Backbuffer = TEXT("Backbuffer");
 };
diff --git a/Engine/Extras/VirtualProduction/TextureShare/TextureShareSDK/Source/Serialize/TextureShareCoreSerialize.h b/Engine/Extras/VirtualProduction/TextureShare/TextureShareSDK/Source/Serialize/TextureShareCoreSerialize.h
--- a/Engine/Extras/VirtualProduction/TextureShare/TextureShareSDK/Source/Serialize/TextureShareCoreSerialize.h
+++ b/Engine/Extras/VirtualProduction/TextureShare/TextureShareSDK/Source/Serialize/TextureShareCoreSerialize.h
@@ -70,15 +70,21 @@ public:
 	{
 		for (const T& ValueIt : InArrayValues)
 		{
-			if (this->Find(ValueIt) == INDEX_NONE)
+			AddSorted(ValueIt);
+		}
+	}
+
+	// Add value ordered. Value must support compare ops
+	void AddSorted(const T InValue)
+	{
+		if (this->Find(InValue) == INDEX_NONE)
+		{
+			for (int32 Index = this->Num() - 1; Index >= 0; Index--)
 			{
-				for (int32 Index = this->Num() - 1; Index >= 0; Index--)
+				if (this->operator[](Index) < InValue)
 				{
-					if (this->operator[](Index) < ValueIt)
-					{
-						this->Insert(ValueIt, Index + 1);
-						break;
-					}
+					this->Insert(InValue, Index + 1);
+					break;
 				}
 			}
 		}
@@ -91,7 +97,7 @@ public:
 		{
 			for (int32 Index = 0; Index < this->Num(); Index++)
 			{
-				if (this->operator[](Index) != InArrayValues[Index])
+				if (!(this->operator[](Index) == InArrayValues[Index]))
 				{
 					return false;
 				}
diff --git a/Engine/Plugins/VirtualProduction/TextureShare/Source/TextureShare/Private/Blueprints/TextureShareBlueprintContainers.cpp b/Engine/Plugins/VirtualProduction/TextureShare/Source/TextureShare/Private/Blueprints/TextureShareBlueprintContainers.cpp
--- a/Engine/Plugins/VirtualProduction/TextureShare/Source/TextureShare/Private/Blueprints/TextureShareBlueprintContainers.cpp
+++ b/Engine/Plugins/VirtualProduction/TextureShare/Source/TextureShare/Private/Blueprints/TextureShareBlueprintContainers.cpp
@@ -5,17 +5,11 @@
 #include "Misc/TextureShareCoreStrings.h"
 #include "UObject/Package.h"
 
-namespace UE
+namespace UE::TextureShare::BlueprintContainers
 {
-	namespace TextureShare
+	static FString GetValidTextureShareObjectName(const FString& InShareName)
 	{
-		namespace BlueprintContainers
-		{
-			static FString GetValidTextureShareObjectName(const FString& InShareName)
-			{
-				return InShareName.IsEmpty() ? UE::TextureShareCoreStrings::DefaultShareName : InShareName;
-			}
-		}
+		return InShareName.IsEmpty() ? UE::TextureShareCoreStrings::DefaultShareName : InShareName;
 	}
 };
 using namespace UE::TextureShare::BlueprintContainers;
diff --git a/Engine/Plugins/VirtualProduction/TextureShare/Source/TextureShare/Private/Game/Settings/TextureShareSettings.cpp b/Engine/Plugins/VirtualProduction/TextureShare/Source/TextureShare/Private/Game/Settings/TextureShareSettings.cpp
--- a/Engine/Plugins/VirtualProduction/TextureShare/Source/TextureShare/Private/Game/Settings/TextureShareSettings.cpp
+++ b/Engine/Plugins/VirtualProduction/TextureShare/Source/TextureShare/Private/Game/Settings/TextureShareSettings.cpp
@@ -8,27 +8,21 @@
 #include "Misc/Paths.h"
 #include "Misc/ConfigCacheIni.h"
 
-namespace UE
+namespace UE::TextureShare::GameSettings
 {
-	namespace TextureShare
-	{
-		namespace GameSettings
-		{
-			static constexpr auto GameConfigCategory = TEXT("/Script/Plugins.TextureShare");
+	static constexpr auto GameConfigCategory = TEXT("/Script/Plugins.TextureShare");
 
-			static const FString GetConfigDefaultGamePath()
-			{
-				static FString ConfigDefaultGamePath = FString::Printf(TEXT("%sDefaultGame.ini"), *FPaths::SourceConfigDir());
-				return ConfigDefaultGamePath;
-			}
+	static const FString GetConfigDefaultGamePath()
+	{
+		static FString ConfigDefaultGamePath = FString::Printf(TEXT("%sDefaultGame.ini"), *FPaths::SourceConfigDir());
+		return ConfigDefaultGamePath;
+	}
 
-			static void ClearReadOnlyForConfigDefaultGame()
-			{
-				if (FPlatformFileManager::Get().GetPlatformFile().IsReadOnly(*GetConfigDefaultGamePath()))
-				{
-					FPlatformFileManager::Get().GetPlatformFile().SetReadOnly(*GetConfigDefaultGamePath(), false);
-				}
-			}
+	static void ClearReadOnlyForConfigDefaultGame()
+	{
+		if (FPlatformFileManager::Get().GetPlatformFile().IsReadOnly(*GetConfigDefaultGamePath()))
+		{
+			FPlatformFileManager::Get().GetPlatformFile().SetReadOnly(*GetConfigDefaultGamePath(), false);
 		}
 	}
 };
diff --git a/Engine/Plugins/VirtualProduction/TextureShare/Source/TextureShare/Private/Game/ViewExtension/TextureShareSceneViewExtension.cpp b/Engine/Plugins/VirtualProduction/TextureShare/Source/TextureShare/Private/Game/ViewExtension/TextureShareSceneViewExtension.cpp
--- a/Engine/Plugins/VirtualProduction/TextureShare/Source/TextureShare/Private/Game/ViewExtension/TextureShareSceneViewExtension.cpp
+++ b/Engine/Plugins/VirtualProduction/TextureShare/Source/TextureShare/Private/Game/ViewExtension/TextureShareSceneViewExtension.cpp
@@ -8,69 +8,66 @@
 #include "Module/TextureShareLog.h"
 #include "Misc/TextureShareStrings.h"
 
+#include "RenderGraphBuilder.h"
+#include "RenderGraphUtils.h"
 #include "RHI.h"
 #include "RHICommandList.h"
 #include "RHIResources.h"
+#include "SceneRenderTargetParameters.h"
 #include "SceneView.h"
 #include "SceneRendering.h"
 
-namespace UE
+namespace UE::TextureShare::SceneViewExtension
 {
-	namespace TextureShare
+	static void GetTextureShareCoreSceneViewMatrices(const FViewMatrices& InViewMatrices, FTextureShareCoreSceneViewMatrices& OutViewMatrices)
 	{
-		namespace SceneViewExtension
-		{
-			static void GetTextureShareCoreSceneViewMatrices(const FViewMatrices& InViewMatrices, FTextureShareCoreSceneViewMatrices& OutViewMatrices)
-			{
-				OutViewMatrices.ProjectionMatrix = InViewMatrices.GetProjectionMatrix();
-				OutViewMatrices.ProjectionNoAAMatrix = InViewMatrices.GetProjectionNoAAMatrix();
-				OutViewMatrices.ViewMatrix = InViewMatrices.GetViewMatrix();
-				OutViewMatrices.ViewProjectionMatrix = InViewMatrices.GetViewProjectionMatrix();
-				OutViewMatrices.TranslatedViewProjectionMatrix = InViewMatrices.GetTranslatedViewProjectionMatrix();
-				OutViewMatrices.PreViewTranslation = InViewMatrices.GetPreViewTranslation();
-				OutViewMatrices.ViewOrigin = InViewMatrices.GetViewOrigin();
-				OutViewMatrices.ProjectionScale = InViewMatrices.GetProjectionScale();
-				OutViewMatrices.TemporalAAProjectionJitter = InViewMatrices.GetTemporalAAJitter();
-				OutViewMatrices.ScreenScale = InViewMatrices.GetScreenScale();
-			}
+		OutViewMatrices.ProjectionMatrix = InViewMatrices.GetProjectionMatrix();
+		OutViewMatrices.ProjectionNoAAMatrix = InViewMatrices.GetProjectionNoAAMatrix();
+		OutViewMatrices.ViewMatrix = InViewMatrices.GetViewMatrix();
+		OutViewMatrices.ViewProjectionMatrix = InViewMatrices.GetViewProjectionMatrix();
+		OutViewMatrices.TranslatedViewProjectionMatrix = InViewMatrices.GetTranslatedViewProjectionMatrix();
+		OutViewMatrices.PreViewTranslation = InViewMatrices.GetPreViewTranslation();
+		OutViewMatrices.ViewOrigin = InViewMatrices.GetViewOrigin();
+		OutViewMatrices.ProjectionScale = InViewMatrices.GetProjectionScale();
+		OutViewMatrices.TemporalAAProjectionJitter = InViewMatrices.GetTemporalAAJitter();
+		OutViewMatrices.ScreenScale = InViewMatrices.GetScreenScale();
+	}
 
-			static void GetTextureShareCoreSceneView(const FSceneView& InSceneView, FTextureShareCoreSceneView& OutSceneView)
-			{
-				GetTextureShareCoreSceneViewMatrices(InSceneView.ViewMatrices, OutSceneView.ViewMatrices);
+	static void GetTextureShareCoreSceneView(const FSceneView& InSceneView, FTextureShareCoreSceneView& OutSceneView)
+	{
+		GetTextureShareCoreSceneViewMatrices(InSceneView.ViewMatrices, OutSceneView.ViewMatrices);
 
-				OutSceneView.UnscaledViewRect = InSceneView.UnscaledViewRect;
-				OutSceneView.UnconstrainedViewRect = InSceneView.UnconstrainedViewRect;
-				OutSceneView.ViewLocation = InSceneView.ViewLocation;
-				OutSceneView.ViewRotation = InSceneView.ViewRotation;
-				OutSceneView.BaseHmdOrientation = InSceneView.BaseHmdOrientation;
-				OutSceneView.BaseHmdLocation = InSceneView.BaseHmdLocation;
-				OutSceneView.WorldToMetersScale = InSceneView.WorldToMetersScale;
+		OutSceneView.UnscaledViewRect = InSceneView.UnscaledViewRect;
+		OutSceneView.UnconstrainedViewRect = InSceneView.UnconstrainedViewRect;
+		OutSceneView.ViewLocation = InSceneView.ViewLocation;
+		OutSceneView.ViewRotation = InSceneView.ViewRotation;
+		OutSceneView.BaseHmdOrientation = InSceneView.BaseHmdOrientation;
+		OutSceneView.BaseHmdLocation = InSceneView.BaseHmdLocation;
+		OutSceneView.WorldToMetersScale = InSceneView.WorldToMetersScale;
 
-				OutSceneView.StereoViewIndex = InSceneView.StereoViewIndex;
-				OutSceneView.PrimaryViewIndex = InSceneView.PrimaryViewIndex;
+		OutSceneView.StereoViewIndex = InSceneView.StereoViewIndex;
+		OutSceneView.PrimaryViewIndex = InSceneView.PrimaryViewIndex;
 
-				OutSceneView.FOV = InSceneView.FOV;
-				OutSceneView.DesiredFOV = InSceneView.DesiredFOV;
-			}
+		OutSceneView.FOV = InSceneView.FOV;
+		OutSceneView.DesiredFOV = InSceneView.DesiredFOV;
+	}
 
-			static void GetTextureShareCoreSceneGameTime(const FGameTime& InGameTime, FTextureShareCoreSceneGameTime& OutGameTime)
-			{
-				OutGameTime.RealTimeSeconds = InGameTime.GetRealTimeSeconds();
-				OutGameTime.WorldTimeSeconds = InGameTime.GetWorldTimeSeconds();
-				OutGameTime.DeltaRealTimeSeconds = InGameTime.GetDeltaRealTimeSeconds();
-				OutGameTime.DeltaWorldTimeSeconds = InGameTime.GetDeltaWorldTimeSeconds();
-			}
+	static void GetTextureShareCoreSceneGameTime(const FGameTime& InGameTime, FTextureShareCoreSceneGameTime& OutGameTime)
+	{
+		OutGameTime.RealTimeSeconds = InGameTime.GetRealTimeSeconds();
+		OutGameTime.WorldTimeSeconds = InGameTime.GetWorldTimeSeconds();
+		OutGameTime.DeltaRealTimeSeconds = InGameTime.GetDeltaRealTimeSeconds();
+		OutGameTime.DeltaWorldTimeSeconds = InGameTime.GetDeltaWorldTimeSeconds();
+	}
 
-			static void GetTextureShareCoreSceneViewFamily(const FSceneViewFamily& InViewFamily, FTextureShareCoreSceneViewFamily& OutViewFamily)
-			{
-				GetTextureShareCoreSceneGameTime(InViewFamily.Time, OutViewFamily.GameTime);
+	static void GetTextureShareCoreSceneViewFamily(const FSceneViewFamily& InViewFamily, FTextureShareCoreSceneViewFamily& OutViewFamily)
+	{
+		GetTextureShareCoreSceneGameTime(InViewFamily.Time, OutViewFamily.GameTime);
 
-				OutViewFamily.FrameNumber = InViewFamily.FrameNumber;
-				OutViewFamily.bIsHDR = InViewFamily.bIsHDR;
-				OutViewFamily.GammaCorrection = InViewFamily.GammaCorrection;
-				OutViewFamily.SecondaryViewFraction = InViewFamily.SecondaryViewFraction;
-			}
-		}
+		OutViewFamily.FrameNumber = InViewFamily.FrameNumber;
+		OutViewFamily.bIsHDR = InViewFamily.bIsHDR;
+		OutViewFamily.GammaCorrection = InViewFamily.GammaCorrection;
+		OutViewFamily.SecondaryViewFraction = InViewFamily.SecondaryViewFraction;
 	}
 };
 using namespace UE::TextureShare::SceneViewExtension;
diff --git a/Engine/Plugins/VirtualProduction/TextureShare/Source/TextureShare/Private/Game/WorldSubsystem/TextureShareWorldSubsystem.cpp b/Engine/Plugins/VirtualProduction/TextureShare/Source/TextureShare/Private/Game/WorldSubsystem/TextureShareWorldSubsystem.cpp
--- a/Engine/Plugins/VirtualProduction/TextureShare/Source/TextureShare/Private/Game/WorldSubsystem/TextureShareWorldSubsystem.cpp
+++ b/Engine/Plugins/VirtualProduction/TextureShare/Source/TextureShare/Private/Game/WorldSubsystem/TextureShareWorldSubsystem.cpp
@@ -22,18 +22,12 @@ static TAutoConsoleVariable<int32> CVarTextureShareUseWorldSubsystem(
 	ECVF_RenderThreadSafe
 );
 
-namespace UE
+namespace UE::TextureShare::WorldSubsystem
 {
-	namespace TextureShare
+	static ITextureShareAPI& TextureShareAPI()
 	{
-		namespace WorldSubsystem
-		{
-			static ITextureShareAPI& TextureShareAPI()
-			{
-				static ITextureShareAPI& TextureShareAPISingletone = ITextureShare::Get().GetTextureShareAPI();
-				return TextureShareAPISingletone;
-			}
-		}
+		static ITextureShareAPI& TextureShareAPISingletone = ITextureShare::Get().GetTextureShareAPI();
+		return TextureShareAPISingletone;
 	}
 };
 using namespace UE::TextureShare::WorldSubsystem;
diff --git a/Engine/Plugins/VirtualProduction/TextureShare/Source/TextureShare/Private/Misc/TextureShareStrings.h b/Engine/Plugins/VirtualProduction/TextureShare/Source/TextureShare/Private/Misc/TextureShareStrings.h
--- a/Engine/Plugins/VirtualProduction/TextureShare/Source/TextureShare/Private/Misc/TextureShareStrings.h
+++ b/Engine/Plugins/VirtualProduction/TextureShare/Source/TextureShare/Private/Misc/TextureShareStrings.h
@@ -10,26 +10,20 @@
 /**
  * Default names for TextureShare module
  */
-namespace UE
+namespace UE::TextureShareStrings::SceneTextures
 {
-	namespace TextureShareStrings
-	{
-		namespace SceneTextures
-		{
-			// Read-only scene resources
-			static constexpr auto SceneColor = TEXT("SceneColor");
-			static constexpr auto SceneDepth = TEXT("SceneDepth");
-			static constexpr auto SmallDepthZ = TEXT("SmallDepthZ");
-			static constexpr auto GBufferA = TEXT("GBufferA");
-			static constexpr auto GBufferB = TEXT("GBufferB");
-			static constexpr auto GBufferC = TEXT("GBufferC");
-			static constexpr auto GBufferD = TEXT("GBufferD");
-			static constexpr auto GBufferE = TEXT("GBufferE");
-			static constexpr auto GBufferF = TEXT("GBufferF");
+	// Read-only scene resources
+	static constexpr auto SceneColor = TEXT("SceneColor");
+	static constexpr auto SceneDepth = TEXT("SceneDepth");
+	static constexpr auto SmallDepthZ = TEXT("SmallDepthZ");
+	static constexpr auto GBufferA = TEXT("GBufferA");
+	static constexpr auto GBufferB = TEXT("GBufferB");
+	static constexpr auto GBufferC = TEXT("GBufferC");
+	static constexpr auto GBufferD = TEXT("GBufferD");
+	static constexpr auto GBufferE = TEXT("GBufferE");
+	static constexpr auto GBufferF = TEXT("GBufferF");
 
-			// Read-write RTTs
-			static constexpr auto FinalColor = TEXT("FinalColor");
-			static constexpr auto Backbuffer = TEXT("Backbuffer");
-		}
-	}
+	// Read-write RTTs
+	static constexpr auto FinalColor = TEXT("FinalColor");
+	static constexpr auto Backbuffer = TEXT("Backbuffer");
 };
diff --git a/Engine/Plugins/VirtualProduction/TextureShare/Source/TextureShare/Private/Module/TextureShareAPI.cpp b/Engine/Plugins/VirtualProduction/TextureShare/Source/TextureShare/Private/Module/TextureShareAPI.cpp
--- a/Engine/Plugins/VirtualProduction/TextureShare/Source/TextureShare/Private/Module/TextureShareAPI.cpp
+++ b/Engine/Plugins/VirtualProduction/TextureShare/Source/TextureShare/Private/Module/TextureShareAPI.cpp
@@ -13,35 +13,29 @@
 #include "Framework/Application/SlateApplication.h"
 
 //////////////////////////////////////////////////////////////////////////////////////////////
-namespace UE
+namespace UE::TextureShare::APIHelpers
 {
-	namespace TextureShare
-	{
-		namespace APIHelpers
-		{
-			static const FName RendererModuleName(TEXT("Renderer"));
+	static const FName RendererModuleName(TEXT("Renderer"));
 
-			static ITextureShareCoreAPI& TextureShareCoreAPI()
-			{
-				static ITextureShareCoreAPI& TextureShareCoreAPISingleton = ITextureShareCore::Get().GetTextureShareCoreAPI();
-				return TextureShareCoreAPISingleton;
-			}
+	static ITextureShareCoreAPI& TextureShareCoreAPI()
+	{
+		static ITextureShareCoreAPI& TextureShareCoreAPISingleton = ITextureShareCore::Get().GetTextureShareCoreAPI();
+		return TextureShareCoreAPISingleton;
+	}
 
-			static ETextureShareDeviceType GetTextureShareDeviceType()
-			{
-				switch (RHIGetInterfaceType())
-				{
-				case ERHIInterfaceType::D3D11:  return ETextureShareDeviceType::D3D11;
-				case ERHIInterfaceType::D3D12:  return ETextureShareDeviceType::D3D12;
-				case ERHIInterfaceType::Vulkan: return ETextureShareDeviceType::Vulkan;
+	static ETextureShareDeviceType GetTextureShareDeviceType()
+	{
+		switch (RHIGetInterfaceType())
+		{
+		case ERHIInterfaceType::D3D11:  return ETextureShareDeviceType::D3D11;
+		case ERHIInterfaceType::D3D12:  return ETextureShareDeviceType::D3D12;
+		case ERHIInterfaceType::Vulkan: return ETextureShareDeviceType::Vulkan;
 		
-				default:
-					break;
-				}
-
-				return ETextureShareDeviceType::Undefined;
-			}
+		default:
+			break;
 		}
+
+		return ETextureShareDeviceType::Undefined;
 	}
 };
 using namespace UE::TextureShare::APIHelpers;
diff --git a/Engine/Plugins/VirtualProduction/TextureShare/Source/TextureShare/Private/Resources/TextureShareResourceSettings.cpp b/Engine/Plugins/VirtualProduction/TextureShare/Source/TextureShare/Private/Resources/TextureShareResourceSettings.cpp
--- a/Engine/Plugins/VirtualProduction/TextureShare/Source/TextureShare/Private/Resources/TextureShareResourceSettings.cpp
+++ b/Engine/Plugins/VirtualProduction/TextureShare/Source/TextureShare/Private/Resources/TextureShareResourceSettings.cpp
@@ -9,243 +9,236 @@
 /**
  * Converters between DXGI and EPixelFormat formats
  */
-namespace UE
+namespace UE::TextureShare::PixelFormats
 {
-	namespace TextureShare
+	static bool GetPixelFormatFromDXGI(const DXGI_FORMAT InFormat, EPixelFormat& OutPixelFormat)
 	{
-		namespace PixelFormats
+		if (InFormat != DXGI_FORMAT_UNKNOWN)
 		{
-			static bool GetPixelFormatFromDXGI(const DXGI_FORMAT InFormat, EPixelFormat& OutPixelFormat)
+			for (uint32 PixelFormat = 0; PixelFormat < EPixelFormat::PF_MAX; ++PixelFormat)
 			{
-				if (InFormat != DXGI_FORMAT_UNKNOWN)
+				if (GPixelFormats[PixelFormat].Supported
+					&& GPixelFormats[PixelFormat].PlatformFormat == InFormat)
 				{
-					for (uint32 PixelFormat = 0; PixelFormat < EPixelFormat::PF_MAX; ++PixelFormat)
-					{
-						if (GPixelFormats[PixelFormat].Supported
-							&& GPixelFormats[PixelFormat].PlatformFormat == InFormat)
-						{
-							// Found
-							OutPixelFormat = (EPixelFormat)PixelFormat;
-
-							return true;
-						}
-					}
+					// Found
+					OutPixelFormat = (EPixelFormat)PixelFormat;
+
+					return true;
 				}
+			}
+		}
+
+		return false;
+	};
+
+	struct FDXGIFormatMap
+	{
+		DXGI_FORMAT           UnormFormat = DXGI_FORMAT_UNKNOWN;
+		DXGI_FORMAT           TypelessFormat = DXGI_FORMAT_UNKNOWN;
 
-				return false;
-			};
+		TArray<DXGI_FORMAT> FullyTypedFormats;
 
-			struct FDXGIFormatMap
+		EPixelFormat FindSharedPixelFormat() const
+		{
+			EPixelFormat PixelFormat;
+
+			if (GetPixelFormatFromDXGI(UnormFormat, PixelFormat))
 			{
-				DXGI_FORMAT           UnormFormat = DXGI_FORMAT_UNKNOWN;
-				DXGI_FORMAT           TypelessFormat = DXGI_FORMAT_UNKNOWN;
+				return PixelFormat;
+			}
 
-				TArray<DXGI_FORMAT> FullyTypedFormats;
+			if (GetPixelFormatFromDXGI(TypelessFormat, PixelFormat))
+			{
+				return PixelFormat;
+			}
 
-				EPixelFormat FindSharedPixelFormat() const
+			for (const DXGI_FORMAT& FormatIt : FullyTypedFormats)
+			{
+				if (GetPixelFormatFromDXGI(FormatIt, PixelFormat))
 				{
-					EPixelFormat PixelFormat;
-
-					if (GetPixelFormatFromDXGI(UnormFormat, PixelFormat))
-					{
-						return PixelFormat;
-					}
-
-					if (GetPixelFormatFromDXGI(TypelessFormat, PixelFormat))
-					{
-						return PixelFormat;
-					}
-
-					for (const DXGI_FORMAT& FormatIt : FullyTypedFormats)
-					{
-						if (GetPixelFormatFromDXGI(FormatIt, PixelFormat))
-						{
-							return PixelFormat;
-						}
-					}
-
-					return PF_Unknown;
+					return PixelFormat;
 				}
-			};
+			}
 
-			// Last supported DXGI format
-			static const DXGI_FORMAT DXGI_FORMAT_MAX = DXGI_FORMAT_B8G8R8X8_UNORM_SRGB;
+			return PF_Unknown;
+		}
+	};
 
-			// Grouped formats: Unorm + typeless + fully typed
-			// Format in group can be copied
-			static const TArray<FDXGIFormatMap> DXGIFormatGroups = {
-				{
-					DXGI_FORMAT_R32G32B32A32_UINT, DXGI_FORMAT_R32G32B32A32_TYPELESS,
-					{
-						DXGI_FORMAT_R32G32B32A32_FLOAT,
-						DXGI_FORMAT_R32G32B32A32_SINT
-					}
-				},
-				{
-					DXGI_FORMAT_R32G32B32_UINT, DXGI_FORMAT_R32G32B32_TYPELESS,
-					{
-						DXGI_FORMAT_R32G32B32_FLOAT,
-						DXGI_FORMAT_R32G32B32_SINT
-					}
-				},
-				{
-					DXGI_FORMAT_R16G16B16A16_UNORM, DXGI_FORMAT_R16G16B16A16_TYPELESS,
-					{
-						DXGI_FORMAT_R16G16B16A16_FLOAT,
-						DXGI_FORMAT_R16G16B16A16_UINT,
-						DXGI_FORMAT_R16G16B16A16_SNORM,
-						DXGI_FORMAT_R16G16B16A16_SINT
-					}
-				},
-				{
-					DXGI_FORMAT_R32G32_UINT, DXGI_FORMAT_R32G32_TYPELESS,
-					{
-						DXGI_FORMAT_R32G32_FLOAT,
-						DXGI_FORMAT_R32G32_SINT
-					}
-				},
-				{
-					DXGI_FORMAT_R10G10B10A2_UNORM, DXGI_FORMAT_R10G10B10A2_TYPELESS,
-					{
-						DXGI_FORMAT_R10G10B10A2_UINT
-					}
-				},
-				{
-					DXGI_FORMAT_R8G8B8A8_UNORM, DXGI_FORMAT_R8G8B8A8_TYPELESS,
-					{
-						DXGI_FORMAT_R8G8B8A8_UNORM_SRGB,
-						DXGI_FORMAT_R8G8B8A8_UINT,
-						DXGI_FORMAT_R8G8B8A8_SNORM,
-						DXGI_FORMAT_R8G8B8A8_SINT
-					}
-				},
-				{
-					DXGI_FORMAT_R16G16_UNORM, DXGI_FORMAT_R16G16_TYPELESS,
-					{
-						DXGI_FORMAT_R16G16_FLOAT,
-						DXGI_FORMAT_R16G16_UINT,
-						DXGI_FORMAT_R16G16_SNORM,
-						DXGI_FORMAT_R16G16_SINT
-					}
-				},
-				{
-					DXGI_FORMAT_R32_UINT, DXGI_FORMAT_R32_TYPELESS,
-					{
-						DXGI_FORMAT_D32_FLOAT,
-						DXGI_FORMAT_R32_FLOAT,
-						DXGI_FORMAT_R32_SINT
-					}
-				},
-				{
-					DXGI_FORMAT_R8G8_UNORM, DXGI_FORMAT_R8G8_TYPELESS,
-					{
-						DXGI_FORMAT_R8G8_TYPELESS,
-						DXGI_FORMAT_R8G8_UNORM,
-						DXGI_FORMAT_R8G8_UINT,
-						DXGI_FORMAT_R8G8_SNORM,
-						DXGI_FORMAT_R8G8_SINT
-					}
-				},
-				{
-					DXGI_FORMAT_R16_UNORM, DXGI_FORMAT_R16_TYPELESS,
-					{
-						DXGI_FORMAT_R16_FLOAT,
-						DXGI_FORMAT_D16_UNORM,
-						DXGI_FORMAT_R16_UINT,
-						DXGI_FORMAT_R16_SNORM,
-						DXGI_FORMAT_R16_SINT
-					}
-				},
-				{
-					DXGI_FORMAT_R8_UNORM, DXGI_FORMAT_R8_TYPELESS,
-					{
-						DXGI_FORMAT_R8_UINT,
-						DXGI_FORMAT_R8_SNORM,
-						DXGI_FORMAT_R8_SINT
-					}
-				},
-				{
-					DXGI_FORMAT_A8_UNORM, DXGI_FORMAT_UNKNOWN,
-					{ }
-				},
-				{
-					DXGI_FORMAT_B8G8R8X8_UNORM, DXGI_FORMAT_B8G8R8X8_TYPELESS,
-					{
-						DXGI_FORMAT_B8G8R8X8_UNORM_SRGB
-					}
-				}
-			};
+	// Last supported DXGI format
+	static const DXGI_FORMAT DXGI_FORMAT_MAX = DXGI_FORMAT_B8G8R8X8_UNORM_SRGB;
 
-			static const FDXGIFormatMap* FindDXGIFormatMap(const DXGI_FORMAT InFormat)
+	// Grouped formats: Unorm + typeless + fully typed
+	// Format in group can be copied
+	static const TArray<FDXGIFormatMap> DXGIFormatGroups = {
+		{
+			DXGI_FORMAT_R32G32B32A32_UINT, DXGI_FORMAT_R32G32B32A32_TYPELESS,
+			{
+				DXGI_FORMAT_R32G32B32A32_FLOAT,
+				DXGI_FORMAT_R32G32B32A32_SINT
+			}
+		},
+		{
+			DXGI_FORMAT_R32G32B32_UINT, DXGI_FORMAT_R32G32B32_TYPELESS,
+			{
+				DXGI_FORMAT_R32G32B32_FLOAT,
+				DXGI_FORMAT_R32G32B32_SINT
+			}
+		},
+		{
+			DXGI_FORMAT_R16G16B16A16_UNORM, DXGI_FORMAT_R16G16B16A16_TYPELESS,
+			{
+				DXGI_FORMAT_R16G16B16A16_FLOAT,
+				DXGI_FORMAT_R16G16B16A16_UINT,
+				DXGI_FORMAT_R16G16B16A16_SNORM,
+				DXGI_FORMAT_R16G16B16A16_SINT
+			}
+		},
+		{
+			DXGI_FORMAT_R32G32_UINT, DXGI_FORMAT_R32G32_TYPELESS,
+			{
+				DXGI_FORMAT_R32G32_FLOAT,
+				DXGI_FORMAT_R32G32_SINT
+			}
+		},
+		{
+			DXGI_FORMAT_R10G10B10A2_UNORM, DXGI_FORMAT_R10G10B10A2_TYPELESS,
+			{
+				DXGI_FORMAT_R10G10B10A2_UINT
+			}
+		},
+		{
+			DXGI_FORMAT_R8G8B8A8_UNORM, DXGI_FORMAT_R8G8B8A8_TYPELESS,
+			{
+				DXGI_FORMAT_R8G8B8A8_UNORM_SRGB,
+				DXGI_FORMAT_R8G8B8A8_UINT,
+				DXGI_FORMAT_R8G8B8A8_SNORM,
+				DXGI_FORMAT_R8G8B8A8_SINT
+			}
+		},
+		{
+			DXGI_FORMAT_R16G16_UNORM, DXGI_FORMAT_R16G16_TYPELESS,
+			{
+				DXGI_FORMAT_R16G16_FLOAT,
+				DXGI_FORMAT_R16G16_UINT,
+				DXGI_FORMAT_R16G16_SNORM,
+				DXGI_FORMAT_R16G16_SINT
+			}
+		},
+		{
+			DXGI_FORMAT_R32_UINT, DXGI_FORMAT_R32_TYPELESS,
+			{
+				DXGI_FORMAT_D32_FLOAT,
+				DXGI_FORMAT_R32_FLOAT,
+				DXGI_FORMAT_R32_SINT
+			}
+		},
+		{
+			DXGI_FORMAT_R8G8_UNORM, DXGI_FORMAT_R8G8_TYPELESS,
+			{
+				DXGI_FORMAT_R8G8_TYPELESS,
+				DXGI_FORMAT_R8G8_UNORM,
+				DXGI_FORMAT_R8G8_UINT,
+				DXGI_FORMAT_R8G8_SNORM,
+				DXGI_FORMAT_R8G8_SINT
+			}
+		},
+		{
+			DXGI_FORMAT_R16_UNORM, DXGI_FORMAT_R16_TYPELESS,
 			{
-				if (InFormat != DXGI_FORMAT_UNKNOWN)
+				DXGI_FORMAT_R16_FLOAT,
+				DXGI_FORMAT_D16_UNORM,
+				DXGI_FORMAT_R16_UINT,
+				DXGI_FORMAT_R16_SNORM,
+				DXGI_FORMAT_R16_SINT
+			}
+		},
+		{
+			DXGI_FORMAT_R8_UNORM, DXGI_FORMAT_R8_TYPELESS,
+			{
+				DXGI_FORMAT_R8_UINT,
+				DXGI_FORMAT_R8_SNORM,
+				DXGI_FORMAT_R8_SINT
+			}
+		},
+		{
+			DXGI_FORMAT_A8_UNORM, DXGI_FORMAT_UNKNOWN,
+			{ }
+		},
+		{
+			DXGI_FORMAT_B8G8R8X8_UNORM, DXGI_FORMAT_B8G8R8X8_TYPELESS,
+			{
+				DXGI_FORMAT_B8G8R8X8_UNORM_SRGB
+			}
+		}
+	};
+
+	static const FDXGIFormatMap* FindDXGIFormatMap(const DXGI_FORMAT InFormat)
+	{
+		if (InFormat != DXGI_FORMAT_UNKNOWN)
+		{
+			return DXGIFormatGroups.FindByPredicate([InFormat](const FDXGIFormatMap& In)
 				{
-					return DXGIFormatGroups.FindByPredicate([InFormat](const FDXGIFormatMap& In)
-						{
-							return In.TypelessFormat == InFormat
-							|| In.UnormFormat == InFormat
-						|| (In.FullyTypedFormats.Find(InFormat) != INDEX_NONE);
-						});
-				}
+					return In.TypelessFormat == InFormat
+					|| In.UnormFormat == InFormat
+				|| (In.FullyTypedFormats.Find(InFormat) != INDEX_NONE);
+				});
+		}
 
-				return nullptr;
-			}
+		return nullptr;
+	}
+
+	/**
+		* Pixel<->DXGI format converter (DirectX texture exchange only supports a limited set of formats)
+		*/
+	class FTextureSharePixelFormats
+	{
+	public:
+		FTextureSharePixelFormats()
+		{
+			DXGIFormatsMap.AddZeroed(DXGI_FORMAT_MAX);
 
-			/**
-			 * Pixel<->DXGI format converter (DirectX texture exchange only supports a limited set of formats)
-			 */
-			class FTextureSharePixelFormats
+			for (int32 DXGIFormatIndex = 0; DXGIFormatIndex < DXGI_FORMAT_MAX; ++DXGIFormatIndex)
 			{
-			public:
-				FTextureSharePixelFormats()
+				DXGI_FORMAT DXGIFormatIt = (DXGI_FORMAT)DXGIFormatIndex;
+				if (const FDXGIFormatMap* DXGIFormatMap = FindDXGIFormatMap(DXGIFormatIt))
 				{
-					DXGIFormatsMap.AddZeroed(DXGI_FORMAT_MAX);
-
-					for (int32 DXGIFormatIndex = 0; DXGIFormatIndex < DXGI_FORMAT_MAX; ++DXGIFormatIndex)
-					{
-						DXGI_FORMAT DXGIFormatIt = (DXGI_FORMAT)DXGIFormatIndex;
-						if (const FDXGIFormatMap* DXGIFormatMap = FindDXGIFormatMap(DXGIFormatIt))
-						{
-							DXGIFormatsMap[DXGIFormatIndex] = DXGIFormatMap->FindSharedPixelFormat();
-						}
-						else
-						{
-							DXGIFormatsMap[DXGIFormatIndex] = PF_Unknown;
-						}
-					}
+					DXGIFormatsMap[DXGIFormatIndex] = DXGIFormatMap->FindSharedPixelFormat();
 				}
-
-				~FTextureSharePixelFormats() = default;
-
-				static EPixelFormat GetSharedPixelFormatFromDXGI(const DXGI_FORMAT InDXGIFormat)
+				else
 				{
-					if (InDXGIFormat != DXGI_FORMAT_UNKNOWN)
-					{
-						static FTextureSharePixelFormats TextureSharePixelFormats;
+					DXGIFormatsMap[DXGIFormatIndex] = PF_Unknown;
+				}
+			}
+		}
 
-						return TextureSharePixelFormats.FindPixelFormat(InDXGIFormat);
-					}
+		~FTextureSharePixelFormats() = default;
 
-					return PF_Unknown;
-				}
+		static EPixelFormat GetSharedPixelFormatFromDXGI(const DXGI_FORMAT InDXGIFormat)
+		{
+			if (InDXGIFormat != DXGI_FORMAT_UNKNOWN)
+			{
+				static FTextureSharePixelFormats TextureSharePixelFormats;
 
-			private:
-				EPixelFormat FindPixelFormat(const DXGI_FORMAT InDXGIFormat) const
-				{
-					if ((int32)(InDXGIFormat) >= 0 && (int32)(InDXGIFormat) < DXGIFormatsMap.Num())
-					{
-						return DXGIFormatsMap[(int32)(InDXGIFormat)];
-					}
+				return TextureSharePixelFormats.FindPixelFormat(InDXGIFormat);
+			}
 
-					return PF_Unknown;
-				}
+			return PF_Unknown;
+		}
 
-			private:
-				TArray<EPixelFormat> DXGIFormatsMap;
-			};
+	private:
+		EPixelFormat FindPixelFormat(const DXGI_FORMAT InDXGIFormat) const
+		{
+			if ((int32)(InDXGIFormat) >= 0 && (int32)(InDXGIFormat) < DXGIFormatsMap.Num())
+			{
+				return DXGIFormatsMap[(int32)(InDXGIFormat)];
+			}
 
+			return PF_Unknown;
 		}
-	}
+
+	private:
+		TArray<EPixelFormat> DXGIFormatsMap;
+	};
 };
 using namespace UE::TextureShare::PixelFormats;
 
diff --git a/Engine/Plugins/VirtualProduction/TextureShare/Source/TextureShare/Private/Resources/TextureShareResource_D3D11.cpp b/Engine/Plugins/VirtualProduction/TextureShare/Source/TextureShare/Private/Resources/TextureShareResource_D3D11.cpp
--- a/Engine/Plugins/VirtualProduction/TextureShare/Source/TextureShare/Private/Resources/TextureShareResource_D3D11.cpp
+++ b/Engine/Plugins/VirtualProduction/TextureShare/Source/TextureShare/Private/Resources/TextureShareResource_D3D11.cpp
@@ -10,20 +10,13 @@
 
 #include "ID3D11DynamicRHI.h"
 
-namespace UE
+namespace UE::TextureShare::Resource_D3D11
 {
-	namespace TextureShare
+	static TSharedPtr<ITextureShareCoreD3D11ResourcesCache, ESPMode::ThreadSafe> GetD3D11ResourcesCache()
 	{
-		namespace Resource_D3D11
-		{
-			static TSharedPtr<ITextureShareCoreD3D11ResourcesCache, ESPMode::ThreadSafe> GetD3D11ResourcesCache()
-			{
-				static ITextureShareCoreAPI& TextureShareCoreAPI = ITextureShareCore::Get().GetTextureShareCoreAPI();
+		static ITextureShareCoreAPI& TextureShareCoreAPI = ITextureShareCore::Get().GetTextureShareCoreAPI();
 
-				return TextureShareCoreAPI.GetD3D11ResourcesCache();
-			}
-
-		}
+		return TextureShareCoreAPI.GetD3D11ResourcesCache();
 	}
 };
 using namespace UE::TextureShare::Resource_D3D11;
diff --git a/Engine/Plugins/VirtualProduction/TextureShare/Source/TextureShare/Private/Resources/TextureShareResource_D3D12.cpp b/Engine/Plugins/VirtualProduction/TextureShare/Source/TextureShare/Private/Resources/TextureShareResource_D3D12.cpp
--- a/Engine/Plugins/VirtualProduction/TextureShare/Source/TextureShare/Private/Resources/TextureShareResource_D3D12.cpp
+++ b/Engine/Plugins/VirtualProduction/TextureShare/Source/TextureShare/Private/Resources/TextureShareResource_D3D12.cpp
@@ -12,73 +12,67 @@
 
 #include "ID3D12DynamicRHI.h"
 
-namespace UE
+namespace UE::TextureShare::Resource_D3D12
 {
-	namespace TextureShare
+	// This function latter may be public in D3D12RHI
+	inline DXGI_FORMAT FindSharedResourceDXGIFormat(DXGI_FORMAT InFormat, bool bSRGB)
 	{
-		namespace Resource_D3D12
+		if (bSRGB)
 		{
-			// This function latter may be public in D3D12RHI
-			inline DXGI_FORMAT FindSharedResourceDXGIFormat(DXGI_FORMAT InFormat, bool bSRGB)
+			switch (InFormat)
 			{
-				if (bSRGB)
-				{
-					switch (InFormat)
-					{
-					case DXGI_FORMAT_B8G8R8X8_TYPELESS:    return DXGI_FORMAT_B8G8R8X8_UNORM_SRGB;
-					case DXGI_FORMAT_B8G8R8A8_TYPELESS:    return DXGI_FORMAT_B8G8R8A8_UNORM_SRGB;
-					case DXGI_FORMAT_R8G8B8A8_TYPELESS:    return DXGI_FORMAT_R8G8B8A8_UNORM_SRGB;
-					case DXGI_FORMAT_BC1_TYPELESS:         return DXGI_FORMAT_BC1_UNORM_SRGB;
-					case DXGI_FORMAT_BC2_TYPELESS:         return DXGI_FORMAT_BC2_UNORM_SRGB;
-					case DXGI_FORMAT_BC3_TYPELESS:         return DXGI_FORMAT_BC3_UNORM_SRGB;
-					case DXGI_FORMAT_BC7_TYPELESS:         return DXGI_FORMAT_BC7_UNORM_SRGB;
-					};
-				}
-				else
-				{
-					switch (InFormat)
-					{
-					case DXGI_FORMAT_B8G8R8X8_TYPELESS:    return DXGI_FORMAT_B8G8R8X8_UNORM;
-					case DXGI_FORMAT_B8G8R8A8_TYPELESS: return DXGI_FORMAT_B8G8R8A8_UNORM;
-					case DXGI_FORMAT_R8G8B8A8_TYPELESS: return DXGI_FORMAT_R8G8B8A8_UNORM;
-					case DXGI_FORMAT_BC1_TYPELESS:      return DXGI_FORMAT_BC1_UNORM;
-					case DXGI_FORMAT_BC2_TYPELESS:      return DXGI_FORMAT_BC2_UNORM;
-					case DXGI_FORMAT_BC3_TYPELESS:      return DXGI_FORMAT_BC3_UNORM;
-					case DXGI_FORMAT_BC7_TYPELESS:      return DXGI_FORMAT_BC7_UNORM;
-					};
-				}
-				switch (InFormat)
-				{
-				case DXGI_FORMAT_R32G32B32A32_TYPELESS: return DXGI_FORMAT_R32G32B32A32_UINT;
-				case DXGI_FORMAT_R32G32B32_TYPELESS:    return DXGI_FORMAT_R32G32B32_UINT;
-				case DXGI_FORMAT_R16G16B16A16_TYPELESS: return DXGI_FORMAT_R16G16B16A16_UNORM;
-				case DXGI_FORMAT_R32G32_TYPELESS:       return DXGI_FORMAT_R32G32_UINT;
-				case DXGI_FORMAT_R10G10B10A2_TYPELESS:  return DXGI_FORMAT_R10G10B10A2_UNORM;
-				case DXGI_FORMAT_R16G16_TYPELESS:       return DXGI_FORMAT_R16G16_UNORM;
-				case DXGI_FORMAT_R8G8_TYPELESS:         return DXGI_FORMAT_R8G8_UNORM;
-				case DXGI_FORMAT_R8_TYPELESS:           return DXGI_FORMAT_R8_UNORM;
+			case DXGI_FORMAT_B8G8R8X8_TYPELESS:    return DXGI_FORMAT_B8G8R8X8_UNORM_SRGB;
+			case DXGI_FORMAT_B8G8R8A8_TYPELESS:    return DXGI_FORMAT_B8G8R8A8_UNORM_SRGB;
+			case DXGI_FORMAT_R8G8B8A8_TYPELESS:    return DXGI_FORMAT_R8G8B8A8_UNORM_SRGB;
+			case DXGI_FORMAT_BC1_TYPELESS:         return DXGI_FORMAT_BC1_UNORM_SRGB;
+			case DXGI_FORMAT_BC2_TYPELESS:         return DXGI_FORMAT_BC2_UNORM_SRGB;
+			case DXGI_FORMAT_BC3_TYPELESS:         return DXGI_FORMAT_BC3_UNORM_SRGB;
+			case DXGI_FORMAT_BC7_TYPELESS:         return DXGI_FORMAT_BC7_UNORM_SRGB;
+			};
+		}
+		else
+		{
+			switch (InFormat)
+			{
+			case DXGI_FORMAT_B8G8R8X8_TYPELESS:    return DXGI_FORMAT_B8G8R8X8_UNORM;
+			case DXGI_FORMAT_B8G8R8A8_TYPELESS: return DXGI_FORMAT_B8G8R8A8_UNORM;
+			case DXGI_FORMAT_R8G8B8A8_TYPELESS: return DXGI_FORMAT_R8G8B8A8_UNORM;
+			case DXGI_FORMAT_BC1_TYPELESS:      return DXGI_FORMAT_BC1_UNORM;
+			case DXGI_FORMAT_BC2_TYPELESS:      return DXGI_FORMAT_BC2_UNORM;
+			case DXGI_FORMAT_BC3_TYPELESS:      return DXGI_FORMAT_BC3_UNORM;
+			case DXGI_FORMAT_BC7_TYPELESS:      return DXGI_FORMAT_BC7_UNORM;
+			};
+		}
+		switch (InFormat)
+		{
+		case DXGI_FORMAT_R32G32B32A32_TYPELESS: return DXGI_FORMAT_R32G32B32A32_UINT;
+		case DXGI_FORMAT_R32G32B32_TYPELESS:    return DXGI_FORMAT_R32G32B32_UINT;
+		case DXGI_FORMAT_R16G16B16A16_TYPELESS: return DXGI_FORMAT_R16G16B16A16_UNORM;
+		case DXGI_FORMAT_R32G32_TYPELESS:       return DXGI_FORMAT_R32G32_UINT;
+		case DXGI_FORMAT_R10G10B10A2_TYPELESS:  return DXGI_FORMAT_R10G10B10A2_UNORM;
+		case DXGI_FORMAT_R16G16_TYPELESS:       return DXGI_FORMAT_R16G16_UNORM;
+		case DXGI_FORMAT_R8G8_TYPELESS:         return DXGI_FORMAT_R8G8_UNORM;
+		case DXGI_FORMAT_R8_TYPELESS:           return DXGI_FORMAT_R8_UNORM;
 
-				case DXGI_FORMAT_BC4_TYPELESS:         return DXGI_FORMAT_BC4_UNORM;
-				case DXGI_FORMAT_BC5_TYPELESS:         return DXGI_FORMAT_BC5_UNORM;
+		case DXGI_FORMAT_BC4_TYPELESS:         return DXGI_FORMAT_BC4_UNORM;
+		case DXGI_FORMAT_BC5_TYPELESS:         return DXGI_FORMAT_BC5_UNORM;
 
 
 
-				case DXGI_FORMAT_R24G8_TYPELESS: return DXGI_FORMAT_R24_UNORM_X8_TYPELESS;
-				case DXGI_FORMAT_R32_TYPELESS: return DXGI_FORMAT_R32_FLOAT;
-				case DXGI_FORMAT_R16_TYPELESS: return DXGI_FORMAT_R16_UNORM;
-					// Changing Depth Buffers to 32 bit on Dingo as D24S8 is actually implemented as a 32 bit buffer in the hardware
-				case DXGI_FORMAT_R32G8X24_TYPELESS: return DXGI_FORMAT_R32_FLOAT_X8X24_TYPELESS;
-				}
-				return InFormat;
-			}
+		case DXGI_FORMAT_R24G8_TYPELESS: return DXGI_FORMAT_R24_UNORM_X8_TYPELESS;
+		case DXGI_FORMAT_R32_TYPELESS: return DXGI_FORMAT_R32_FLOAT;
+		case DXGI_FORMAT_R16_TYPELESS: return DXGI_FORMAT_R16_UNORM;
+			// Changing Depth Buffers to 32 bit on Dingo as D24S8 is actually implemented as a 32 bit buffer in the hardware
+		case DXGI_FORMAT_R32G8X24_TYPELESS: return DXGI_FORMAT_R32_FLOAT_X8X24_TYPELESS;
+		}
+		return InFormat;
+	}
 
-			static TSharedPtr<ITextureShareCoreD3D12ResourcesCache, ESPMode::ThreadSafe> GetD3D12ResourcesCache()
-			{
-				static ITextureShareCoreAPI& TextureShareCoreAPI = ITextureShareCore::Get().GetTextureShareCoreAPI();
+	static TSharedPtr<ITextureShareCoreD3D12ResourcesCache, ESPMode::ThreadSafe> GetD3D12ResourcesCache()
+	{
+		static ITextureShareCoreAPI& TextureShareCoreAPI = ITextureShareCore::Get().GetTextureShareCoreAPI();
 
-				return TextureShareCoreAPI.GetD3D12ResourcesCache();
-			}
-		}
+		return TextureShareCoreAPI.GetD3D12ResourcesCache();
 	}
 };
 using namespace UE::TextureShare::Resource_D3D12;
diff --git a/Engine/Plugins/VirtualProduction/TextureShare/Source/TextureShare/Private/Resources/TextureShareResource_Vulkan.cpp b/Engine/Plugins/VirtualProduction/TextureShare/Source/TextureShare/Private/Resources/TextureShareResource_Vulkan.cpp
--- a/Engine/Plugins/VirtualProduction/TextureShare/Source/TextureShare/Private/Resources/TextureShareResource_Vulkan.cpp
+++ b/Engine/Plugins/VirtualProduction/TextureShare/Source/TextureShare/Private/Resources/TextureShareResource_Vulkan.cpp
@@ -15,29 +15,23 @@
 #include "RenderResource.h"
 #include "VulkanRHIPrivate.h"
 
-namespace UE
+namespace UE::TextureShare::Resource_Vulkan
 {
-	namespace TextureShare
+	static TSharedPtr<ITextureShareCoreVulkanResourcesCache, ESPMode::ThreadSafe> GetVulkanResourcesCache()
 	{
-		namespace Resource_Vulkan
-		{
-			static TSharedPtr<ITextureShareCoreVulkanResourcesCache, ESPMode::ThreadSafe> GetVulkanResourcesCache()
-			{
-				static ITextureShareCoreAPI& TextureShareCoreAPI = ITextureShareCore::Get().GetTextureShareCoreAPI();
-
-				return TextureShareCoreAPI.GetVulkanResourcesCache();
-			}
+		static ITextureShareCoreAPI& TextureShareCoreAPI = ITextureShareCore::Get().GetTextureShareCoreAPI();
 
-			static FTextureShareDeviceVulkanContext GetDeviceVulkanContext()
-			{
-				if (FVulkanDynamicRHI* DynamicRHI = static_cast<FVulkanDynamicRHI*>(GDynamicRHI))
-				{
-					return FTextureShareDeviceVulkanContext(DynamicRHI->GetInstance(), DynamicRHI->GetDevice()->GetPhysicalHandle(), DynamicRHI->GetDevice()->GetInstanceHandle());
-				}
+		return TextureShareCoreAPI.GetVulkanResourcesCache();
+	}
 
-				return FTextureShareDeviceVulkanContext();
-			}
+	static FTextureShareDeviceVulkanContext GetDeviceVulkanContext()
+	{
+		if (FVulkanDynamicRHI* DynamicRHI = static_cast<FVulkanDynamicRHI*>(GDynamicRHI))
+		{
+			return FTextureShareDeviceVulkanContext(DynamicRHI->GetInstance(), DynamicRHI->GetDevice()->GetPhysicalHandle(), DynamicRHI->GetDevice()->GetInstanceHandle());
 		}
+
+		return FTextureShareDeviceVulkanContext();
 	}
 };
 using namespace UE::TextureShare::Resource_Vulkan;
diff --git a/Engine/Plugins/VirtualProduction/TextureShare/Source/TextureShareCore/Private/Core/ResourcesCache/D3D11/TextureShareCoreD3D11ResourcesCacheItem.cpp b/Engine/Plugins/VirtualProduction/TextureShare/Source/TextureShareCore/Private/Core/ResourcesCache/D3D11/TextureShareCoreD3D11ResourcesCacheItem.cpp
--- a/Engine/Plugins/VirtualProduction/TextureShare/Source/TextureShareCore/Private/Core/ResourcesCache/D3D11/TextureShareCoreD3D11ResourcesCacheItem.cpp
+++ b/Engine/Plugins/VirtualProduction/TextureShare/Source/TextureShareCore/Private/Core/ResourcesCache/D3D11/TextureShareCoreD3D11ResourcesCacheItem.cpp
@@ -6,181 +6,175 @@
 
 #include "Windows/D3D11ThirdParty.h"
 
-namespace UE
+namespace UE::TextureShareCore::D3D11ResourcesCacheItem
 {
-	namespace TextureShareCore
+	static inline const FString GetD3D11ComErrorDescription(HRESULT Res)
 	{
-		namespace D3D11ResourcesCacheItem
+		const uint32 BufSize = 4096;
+		WIDECHAR buffer[4096];
+		if (::FormatMessageW(FORMAT_MESSAGE_FROM_SYSTEM,
+			nullptr,
+			Res,
+			MAKELANGID(LANG_NEUTRAL, SUBLANG_NEUTRAL),
+			buffer,
+			sizeof(buffer) / sizeof(*buffer),
+			nullptr))
 		{
-			static inline const FString GetD3D11ComErrorDescription(HRESULT Res)
-			{
-				const uint32 BufSize = 4096;
-				WIDECHAR buffer[4096];
-				if (::FormatMessageW(FORMAT_MESSAGE_FROM_SYSTEM,
-					nullptr,
-					Res,
-					MAKELANGID(LANG_NEUTRAL, SUBLANG_NEUTRAL),
-					buffer,
-					sizeof(buffer) / sizeof(*buffer),
-					nullptr))
-				{
-					return buffer;
-				}
-				else
-				{
-					return TEXT("[cannot find d3d11 error description]");
-				}
-			}
+			return buffer;
+		}
+		else
+		{
+			return TEXT("[cannot find d3d11 error description]");
+		}
+	}
 
-			static inline FString GetUniqueD3D11ResourceHandleName(const FGuid& InResourceGuid)
+	static inline FString GetUniqueD3D11ResourceHandleName(const FGuid& InResourceGuid)
+	{
+		return FString::Printf(TEXT("Global\\%s"), *InResourceGuid.ToString(EGuidFormats::DigitsWithHyphensInBraces));
+	}
+
+	static bool D3D11CreateSharedNTHandle(ID3D11Texture2D* InResourceD3D11, const void* InSecurityAttributes, FGuid& OutResourceHandleGuid, HANDLE& OutSharedNTHandle)
+	{
+		TRefCountPtr<IDXGIResource1> DXGIResource1;
+		{
+			// Open DXGIResource1 interface:
+			HRESULT Result = InResourceD3D11->QueryInterface(IID_PPV_ARGS(DXGIResource1.GetInitReference()));
+			if (FAILED(Result))
 			{
-				return FString::Printf(TEXT("Global\\%s"), *InResourceGuid.ToString(EGuidFormats::DigitsWithHyphensInBraces));
+				UE_LOG(LogTextureShareCoreD3D, Error, TEXT("D3D11CreateNamedNTHandle::QueryInterface failed: 0x%X - %s"), Result, *GetD3D11ComErrorDescription(Result));
+
+				return false;
 			}
+		}
 
-			static bool D3D11CreateSharedNTHandle(ID3D11Texture2D* InResourceD3D11, const void* InSecurityAttributes, FGuid& OutResourceHandleGuid, HANDLE& OutSharedNTHandle)
-			{
-				TRefCountPtr<IDXGIResource1> DXGIResource1;
-				{
-					// Open DXGIResource1 interface:
-					HRESULT Result = InResourceD3D11->QueryInterface(IID_PPV_ARGS(DXGIResource1.GetInitReference()));
-					if (FAILED(Result))
-					{
-						UE_LOG(LogTextureShareCoreD3D, Error, TEXT("D3D11CreateNamedNTHandle::QueryInterface failed: 0x%X - %s"), Result, *GetD3D11ComErrorDescription(Result));
+		if (DXGIResource1.IsValid())
+		{
+			OutResourceHandleGuid = FGuid::NewGuid();
 
-						return false;
-					}
-				}
+			// Microsoft docs: Currently the only value this parameter accepts is GENERIC_ALL.
+			const DWORD dwAccess = GENERIC_ALL;
 
-				if (DXGIResource1.IsValid())
-				{
-					OutResourceHandleGuid = FGuid::NewGuid();
+			// Create new shared handle
+			HRESULT Result = DXGIResource1->CreateSharedHandle((const SECURITY_ATTRIBUTES*)InSecurityAttributes, dwAccess, *GetUniqueD3D11ResourceHandleName(OutResourceHandleGuid), &OutSharedNTHandle);
+			if (FAILED(Result))
+			{
+				UE_LOG(LogTextureShareCoreD3D, Error, TEXT("D3D11CreateNamedNTHandle::CreateSharedHandle failed: 0x%X - %s"), Result, *GetD3D11ComErrorDescription(Result));
 
-					// Microsoft docs: Currently the only value this parameter accepts is GENERIC_ALL.
-					const DWORD dwAccess = GENERIC_ALL;
+				return false;
+			}
 
-					// Create new shared handle
-					HRESULT Result = DXGIResource1->CreateSharedHandle((const SECURITY_ATTRIBUTES*)InSecurityAttributes, dwAccess, *GetUniqueD3D11ResourceHandleName(OutResourceHandleGuid), &OutSharedNTHandle);
-					if (FAILED(Result))
-					{
-						UE_LOG(LogTextureShareCoreD3D, Error, TEXT("D3D11CreateNamedNTHandle::CreateSharedHandle failed: 0x%X - %s"), Result, *GetD3D11ComErrorDescription(Result));
+			return OutSharedNTHandle != NULL;
+		}
 
-						return false;
-					}
+		UE_LOG(LogTextureShareCoreD3D, Error, TEXT("D3D11CreateNamedNTHandle::DXGIResource1 invalid"));
 
-					return OutSharedNTHandle != NULL;
-				}
+		return false;
+	}
 
-				UE_LOG(LogTextureShareCoreD3D, Error, TEXT("D3D11CreateNamedNTHandle::DXGIResource1 invalid"));
+	static bool D3D11CreateSharedHandle(ID3D11Texture2D* InResourceD3D11, HANDLE& OutSharedHandle)
+	{
+		// Use old D3D11 shared resources interfaces
+		TRefCountPtr<IDXGIResource> DXGIResource;
+		{
+			HRESULT Result = InResourceD3D11->QueryInterface(IID_PPV_ARGS(DXGIResource.GetInitReference()));
+			if (FAILED(Result))
+			{
+				UE_LOG(LogTextureShareCoreD3D, Error, TEXT("D3D11CreateSharedHandle::QueryInterface failed: 0x%X - %s"), Result, *GetD3D11ComErrorDescription(Result));
 
 				return false;
 			}
+		}
 
-			static bool D3D11CreateSharedHandle(ID3D11Texture2D* InResourceD3D11, HANDLE& OutSharedHandle)
+		if (DXGIResource.IsValid())
+		{
+			//
+			// NOTE : The HANDLE IDXGIResource::GetSharedHandle gives us is NOT an NT Handle, and therefre we should not call CloseHandle on it
+			//
+			HRESULT Result = DXGIResource->GetSharedHandle(&OutSharedHandle);
+			if (FAILED(Result))
 			{
-				// Use old D3D11 shared resources interfaces
-				TRefCountPtr<IDXGIResource> DXGIResource;
-				{
-					HRESULT Result = InResourceD3D11->QueryInterface(IID_PPV_ARGS(DXGIResource.GetInitReference()));
-					if (FAILED(Result))
-					{
-						UE_LOG(LogTextureShareCoreD3D, Error, TEXT("D3D11CreateSharedHandle::QueryInterface failed: 0x%X - %s"), Result, *GetD3D11ComErrorDescription(Result));
+				UE_LOG(LogTextureShareCoreD3D, Error, TEXT("D3D11CreateSharedHandle::GetSharedHandle failed: 0x%X - %s"), Result, *GetD3D11ComErrorDescription(Result));
 
-						return false;
-					}
-				}
+				return false;
+			}
 
-				if (DXGIResource.IsValid())
-				{
-					//
-					// NOTE : The HANDLE IDXGIResource::GetSharedHandle gives us is NOT an NT Handle, and therefre we should not call CloseHandle on it
-					//
-					HRESULT Result = DXGIResource->GetSharedHandle(&OutSharedHandle);
-					if (FAILED(Result))
-					{
-						UE_LOG(LogTextureShareCoreD3D, Error, TEXT("D3D11CreateSharedHandle::GetSharedHandle failed: 0x%X - %s"), Result, *GetD3D11ComErrorDescription(Result));
+			return OutSharedHandle != NULL;
+		}
 
-						return false;
-					}
+		UE_LOG(LogTextureShareCoreD3D, Error, TEXT("D3D11CreateSharedHandle::DXGIResource invalid"));
 
-					return OutSharedHandle != NULL;
-				}
+		return false;
+	}
+
+	static ID3D11Texture2D* D3D11OpenSharedHandle(ID3D11Device* InD3D11Device, const HANDLE InSharedResourceHandle)
+	{
+		// Open not NT Handle
+		if (InD3D11Device && InSharedResourceHandle)
+		{
+			ID3D11Texture2D* D3D11Resource = nullptr;
+			HRESULT Result = InD3D11Device->OpenSharedResource(InSharedResourceHandle, __uuidof(ID3D11Texture2D), (LPVOID*)&D3D11Resource);
+			if (SUCCEEDED(Result))
+			{
+				return D3D11Resource;
+			}
 
-				UE_LOG(LogTextureShareCoreD3D, Error, TEXT("D3D11CreateSharedHandle::DXGIResource invalid"));
+			UE_LOG(LogTextureShareCoreD3D, Error, TEXT("D3D11OpenSharedResourceFromHandle::OpenSharedResource failed: 0x%X - %s"), Result, *GetD3D11ComErrorDescription(Result));
+		}
 
-				return false;
+		return nullptr;
+	}
+
+	static ID3D11Texture2D* D3D11OpenSharedNTHandle(ID3D11Device* InD3D11Device, const FGuid& InResourceHandleGuid, const HANDLE& InSharedNTHandle)
+	{
+		const bool bOpenNamedHandle = InResourceHandleGuid.IsValid();
+		if (bOpenNamedHandle || InSharedNTHandle)
+		{
+			// Gives a device access to a shared resource that is referenced by a handle and that was created on a different device.
+			// You must have previously created the resource as shared and specified that it uses NT handles
+			// (that is, you set the D3D11_RESOURCE_MISC_SHARED_NTHANDLE flag).
+			TRefCountPtr <ID3D11Device1> D3D11Device1;
+			{
+				HRESULT Result = InD3D11Device->QueryInterface(__uuidof(ID3D11Device1), (void**)D3D11Device1.GetInitReference());
+				if (FAILED(Result))
+				{
+					UE_LOG(LogTextureShareCoreD3D, Error, TEXT("D3D11OpenSharedNTHandle::QueryInterface failed: 0x%X - %s"), Result, *GetD3D11ComErrorDescription(Result));
+
+					return nullptr;
+				}
 			}
 
-			static ID3D11Texture2D* D3D11OpenSharedHandle(ID3D11Device* InD3D11Device, const HANDLE InSharedResourceHandle)
+			if (D3D11Device1.IsValid())
 			{
-				// Open not NT Handle
-				if (InD3D11Device && InSharedResourceHandle)
+				ID3D11Texture2D* D3D11Resource = nullptr;
+
+				if (bOpenNamedHandle)
 				{
-					ID3D11Texture2D* D3D11Resource = nullptr;
-					HRESULT Result = InD3D11Device->OpenSharedResource(InSharedResourceHandle, __uuidof(ID3D11Texture2D), (LPVOID*)&D3D11Resource);
+					const DWORD dwDesiredAccess = DXGI_SHARED_RESOURCE_READ | DXGI_SHARED_RESOURCE_WRITE;
+
+					HRESULT Result = D3D11Device1->OpenSharedResourceByName(*GetUniqueD3D11ResourceHandleName(InResourceHandleGuid), dwDesiredAccess, __uuidof(ID3D11Texture2D), (LPVOID*)&D3D11Resource);
 					if (SUCCEEDED(Result))
 					{
 						return D3D11Resource;
 					}
 
-					UE_LOG(LogTextureShareCoreD3D, Error, TEXT("D3D11OpenSharedResourceFromHandle::OpenSharedResource failed: 0x%X - %s"), Result, *GetD3D11ComErrorDescription(Result));
+					UE_LOG(LogTextureShareCoreD3D, Warning, TEXT("D3D11OpenSharedNTHandle::OpenSharedResourceByName failed: 0x%X - %s"), Result, *GetD3D11ComErrorDescription(Result));
+					// try to open from NTHandle
 				}
 
-				return nullptr;
-			}
-
-			static ID3D11Texture2D* D3D11OpenSharedNTHandle(ID3D11Device* InD3D11Device, const FGuid& InResourceHandleGuid, const HANDLE& InSharedNTHandle)
-			{
-				const bool bOpenNamedHandle = InResourceHandleGuid.IsValid();
-				if (bOpenNamedHandle || InSharedNTHandle)
+				if (InSharedNTHandle)
 				{
-					// Gives a device access to a shared resource that is referenced by a handle and that was created on a different device.
-					// You must have previously created the resource as shared and specified that it uses NT handles
-					// (that is, you set the D3D11_RESOURCE_MISC_SHARED_NTHANDLE flag).
-					TRefCountPtr <ID3D11Device1> D3D11Device1;
+					HRESULT Result = D3D11Device1->OpenSharedResource1(InSharedNTHandle, __uuidof(ID3D11Texture2D), (LPVOID*)&D3D11Resource);
+					if (SUCCEEDED(Result))
 					{
-						HRESULT Result = InD3D11Device->QueryInterface(__uuidof(ID3D11Device1), (void**)D3D11Device1.GetInitReference());
-						if (FAILED(Result))
-						{
-							UE_LOG(LogTextureShareCoreD3D, Error, TEXT("D3D11OpenSharedNTHandle::QueryInterface failed: 0x%X - %s"), Result, *GetD3D11ComErrorDescription(Result));
-
-							return nullptr;
-						}
+						return D3D11Resource;
 					}
 
-					if (D3D11Device1.IsValid())
-					{
-						ID3D11Texture2D* D3D11Resource = nullptr;
-
-						if (bOpenNamedHandle)
-						{
-							const DWORD dwDesiredAccess = DXGI_SHARED_RESOURCE_READ | DXGI_SHARED_RESOURCE_WRITE;
-
-							HRESULT Result = D3D11Device1->OpenSharedResourceByName(*GetUniqueD3D11ResourceHandleName(InResourceHandleGuid), dwDesiredAccess, __uuidof(ID3D11Texture2D), (LPVOID*)&D3D11Resource);
-							if (SUCCEEDED(Result))
-							{
-								return D3D11Resource;
-							}
-
-							UE_LOG(LogTextureShareCoreD3D, Warning, TEXT("D3D11OpenSharedNTHandle::OpenSharedResourceByName failed: 0x%X - %s"), Result, *GetD3D11ComErrorDescription(Result));
-							// try to open from NTHandle
-						}
-
-						if (InSharedNTHandle)
-						{
-							HRESULT Result = D3D11Device1->OpenSharedResource1(InSharedNTHandle, __uuidof(ID3D11Texture2D), (LPVOID*)&D3D11Resource);
-							if (SUCCEEDED(Result))
-							{
-								return D3D11Resource;
-							}
-
-							UE_LOG(LogTextureShareCoreD3D, Error, TEXT("D3D11OpenSharedNTHandle::OpenSharedResource1 failed: 0x%X - %s"), Result, *GetD3D11ComErrorDescription(Result));
-						}
-					}
+					UE_LOG(LogTextureShareCoreD3D, Error, TEXT("D3D11OpenSharedNTHandle::OpenSharedResource1 failed: 0x%X - %s"), Result, *GetD3D11ComErrorDescription(Result));
 				}
-
-				return nullptr;
 			}
 		}
+
+		return nullptr;
 	}
 };
 using namespace UE::TextureShareCore::D3D11ResourcesCacheItem;
diff --git a/Engine/Plugins/VirtualProduction/TextureShare/Source/TextureShareCore/Private/Core/ResourcesCache/D3D12/TextureShareCoreD3D12ResourcesCacheItem.cpp b/Engine/Plugins/VirtualProduction/TextureShare/Source/TextureShareCore/Private/Core/ResourcesCache/D3D12/TextureShareCoreD3D12ResourcesCacheItem.cpp
--- a/Engine/Plugins/VirtualProduction/TextureShare/Source/TextureShareCore/Private/Core/ResourcesCache/D3D12/TextureShareCoreD3D12ResourcesCacheItem.cpp
+++ b/Engine/Plugins/VirtualProduction/TextureShare/Source/TextureShareCore/Private/Core/ResourcesCache/D3D12/TextureShareCoreD3D12ResourcesCacheItem.cpp
@@ -6,38 +6,32 @@
 
 #include "Windows/D3D12ThirdParty.h"
 
-namespace UE
+namespace UE::TextureShareCore::D3D12ResourcesCacheItem
 {
-	namespace TextureShareCore
+	static inline const FString GetD3D12ComErrorDescription(HRESULT Res)
 	{
-		namespace D3D12ResourcesCacheItem
+		const uint32 BufSize = 4096;
+		WIDECHAR buffer[4096];
+		if (::FormatMessageW(FORMAT_MESSAGE_FROM_SYSTEM,
+			nullptr,
+			Res,
+			MAKELANGID(LANG_NEUTRAL, SUBLANG_NEUTRAL),
+			buffer,
+			sizeof(buffer) / sizeof(*buffer),
+			nullptr))
 		{
-			static inline const FString GetD3D12ComErrorDescription(HRESULT Res)
-			{
-				const uint32 BufSize = 4096;
-				WIDECHAR buffer[4096];
-				if (::FormatMessageW(FORMAT_MESSAGE_FROM_SYSTEM,
-					nullptr,
-					Res,
-					MAKELANGID(LANG_NEUTRAL, SUBLANG_NEUTRAL),
-					buffer,
-					sizeof(buffer) / sizeof(*buffer),
-					nullptr))
-				{
-					return buffer;
-				}
-				else
-				{
-					return TEXT("[cannot find d3d12 error description]");
-				}
-			}
-
-			static inline FString GetUniqueD3D12ResourceHandleName(const FGuid& InResourceGuid)
-			{
-				return FString::Printf(TEXT("Global\\%s"), *InResourceGuid.ToString(EGuidFormats::DigitsWithHyphensInBraces));
-			}
+			return buffer;
+		}
+		else
+		{
+			return TEXT("[cannot find d3d12 error description]");
 		}
 	}
+
+	static inline FString GetUniqueD3D12ResourceHandleName(const FGuid& InResourceGuid)
+	{
+		return FString::Printf(TEXT("Global\\%s"), *InResourceGuid.ToString(EGuidFormats::DigitsWithHyphensInBraces));
+	}
 };
 using namespace UE::TextureShareCore::D3D12ResourcesCacheItem;
 
diff --git a/Engine/Plugins/VirtualProduction/TextureShare/Source/TextureShareCore/Private/Core/ResourcesCache/Vulkan/TextureShareCoreVulkanResourcesCacheItem.cpp b/Engine/Plugins/VirtualProduction/TextureShare/Source/TextureShareCore/Private/Core/ResourcesCache/Vulkan/TextureShareCoreVulkanResourcesCacheItem.cpp
--- a/Engine/Plugins/VirtualProduction/TextureShare/Source/TextureShareCore/Private/Core/ResourcesCache/Vulkan/TextureShareCoreVulkanResourcesCacheItem.cpp
+++ b/Engine/Plugins/VirtualProduction/TextureShare/Source/TextureShareCore/Private/Core/ResourcesCache/Vulkan/TextureShareCoreVulkanResourcesCacheItem.cpp
@@ -31,150 +31,144 @@ __pragma(warning(disable:4191))
 #define GetVulkanDeviceProc(name)\
 	PFN_vk##name name = (PFN_vk##name)GetVulkanLibrary().GetDeviceInstanceProcAddr(InDeviceVulkanContext, "vk"#name)
 
-namespace UE
+namespace UE::TextureShareCore::VulkanResourcesCacheItem
 {
-	namespace TextureShareCore
+	//////////////////////////////////////////////////////////////////////////////////////////////
+	// FTextureShareCoreVulkanLibrary
+	//////////////////////////////////////////////////////////////////////////////////////////////
+	class FTextureShareCoreVulkanLibrary
 	{
-		namespace VulkanResourcesCacheItem
-		{
-			//////////////////////////////////////////////////////////////////////////////////////////////
-			// FTextureShareCoreVulkanLibrary
-			//////////////////////////////////////////////////////////////////////////////////////////////
-			class FTextureShareCoreVulkanLibrary
-			{
-			public:
-				FTextureShareCoreVulkanLibrary();
-				~FTextureShareCoreVulkanLibrary();
+	public:
+		FTextureShareCoreVulkanLibrary();
+		~FTextureShareCoreVulkanLibrary();
 
-			public:
-				bool IsValid() const
-				{
-					return VulkanLibraryPtr != nullptr && vkGetInstanceProcAddr != nullptr;
-				}
+	public:
+		bool IsValid() const
+		{
+			return VulkanLibraryPtr != nullptr && vkGetInstanceProcAddr != nullptr;
+		}
 
-				void* GetInstanceProcAddr(const FTextureShareDeviceVulkanContext& InDeviceContext, LPCSTR InProcName) const;
-				void* GetProcAddr(LPCSTR InProcName) const;
-				void* GetDeviceInstanceProcAddr(const FTextureShareDeviceVulkanContext& InDeviceContext, LPCSTR InProcName) const;
+		void* GetInstanceProcAddr(const FTextureShareDeviceVulkanContext& InDeviceContext, LPCSTR InProcName) const;
+		void* GetProcAddr(LPCSTR InProcName) const;
+		void* GetDeviceInstanceProcAddr(const FTextureShareDeviceVulkanContext& InDeviceContext, LPCSTR InProcName) const;
 
-			protected:
-				// Load the library
-				bool Load();
-				// Unload the library
-				void Unload();
+	protected:
+		// Load the library
+		bool Load();
+		// Unload the library
+		void Unload();
 
-			private:
-				void* VulkanLibraryPtr = nullptr;
-				PFN_vkGetInstanceProcAddr vkGetInstanceProcAddr = nullptr;
-			};
+	private:
+		void* VulkanLibraryPtr = nullptr;
+		PFN_vkGetInstanceProcAddr vkGetInstanceProcAddr = nullptr;
+	};
 
-			//////////////////////////////////////////////////////////////////////////////////////////////
-			FTextureShareCoreVulkanLibrary::FTextureShareCoreVulkanLibrary()
-			{
-				if (!Load())
-				{
-					Unload();
-				}
-			}
+	//////////////////////////////////////////////////////////////////////////////////////////////
+	FTextureShareCoreVulkanLibrary::FTextureShareCoreVulkanLibrary()
+	{
+		if (!Load())
+		{
+			Unload();
+		}
+	}
 
-			FTextureShareCoreVulkanLibrary::~FTextureShareCoreVulkanLibrary()
-			{
-				Unload();
-			}
+	FTextureShareCoreVulkanLibrary::~FTextureShareCoreVulkanLibrary()
+	{
+		Unload();
+	}
 
-			void* FTextureShareCoreVulkanLibrary::GetInstanceProcAddr(const FTextureShareDeviceVulkanContext& InDeviceContext, LPCSTR InProcName) const
-			{
-				if (IsValid() && InProcName && *InProcName)
-				{
-					return vkGetInstanceProcAddr(InDeviceContext.GetInstance(), InProcName);
-				}
+	void* FTextureShareCoreVulkanLibrary::GetInstanceProcAddr(const FTextureShareDeviceVulkanContext& InDeviceContext, LPCSTR InProcName) const
+	{
+		if (IsValid() && InProcName && *InProcName)
+		{
+			return vkGetInstanceProcAddr(InDeviceContext.GetInstance(), InProcName);
+		}
 
-				return nullptr;
-			}
+		return nullptr;
+	}
 
-			void* FTextureShareCoreVulkanLibrary::GetProcAddr(LPCSTR InProcName) const
-			{
-				if (IsValid() && InProcName && *InProcName)
-				{
-					return vkGetInstanceProcAddr(nullptr, InProcName);
-				}
+	void* FTextureShareCoreVulkanLibrary::GetProcAddr(LPCSTR InProcName) const
+	{
+		if (IsValid() && InProcName && *InProcName)
+		{
+			return vkGetInstanceProcAddr(nullptr, InProcName);
+		}
 
-				return nullptr;
-			}
+		return nullptr;
+	}
 
-			void* FTextureShareCoreVulkanLibrary::GetDeviceInstanceProcAddr(const FTextureShareDeviceVulkanContext& InDeviceContext, LPCSTR InProcName) const
-			{
-				const PFN_vkGetDeviceProcAddr GetDeviceProcAddr = (PFN_vkGetDeviceProcAddr)GetInstanceProcAddr(InDeviceContext, "vkGetDeviceProcAddr");
-				if (GetDeviceProcAddr)
-				{
-					return GetDeviceProcAddr(InDeviceContext.GetDevice(), InProcName);
-				}
+	void* FTextureShareCoreVulkanLibrary::GetDeviceInstanceProcAddr(const FTextureShareDeviceVulkanContext& InDeviceContext, LPCSTR InProcName) const
+	{
+		const PFN_vkGetDeviceProcAddr GetDeviceProcAddr = (PFN_vkGetDeviceProcAddr)GetInstanceProcAddr(InDeviceContext, "vkGetDeviceProcAddr");
+		if (GetDeviceProcAddr)
+		{
+			return GetDeviceProcAddr(InDeviceContext.GetDevice(), InProcName);
+		}
 
-				return nullptr;
-			}
+		return nullptr;
+	}
 
-			bool FTextureShareCoreVulkanLibrary::Load()
-			{
-				VulkanLibraryPtr = ::LoadLibraryA(VK_LIBRARY_NAME);
-				if (VulkanLibraryPtr)
-				{
-					vkGetInstanceProcAddr = (PFN_vkGetInstanceProcAddr)(::GetProcAddress((Windows::HMODULE)VulkanLibraryPtr, "vkGetInstanceProcAddr"));
+	bool FTextureShareCoreVulkanLibrary::Load()
+	{
+		VulkanLibraryPtr = ::LoadLibraryA(VK_LIBRARY_NAME);
+		if (VulkanLibraryPtr)
+		{
+			vkGetInstanceProcAddr = (PFN_vkGetInstanceProcAddr)(::GetProcAddress((Windows::HMODULE)VulkanLibraryPtr, "vkGetInstanceProcAddr"));
 
-					if (!vkGetInstanceProcAddr)
-						return false;
+			if (!vkGetInstanceProcAddr)
+				return false;
 
-					return true;
-				}
+			return true;
+		}
 
-				return false;
-			}
+		return false;
+	}
 
-			void FTextureShareCoreVulkanLibrary::Unload()
-			{
-				if (VulkanLibraryPtr)
-				{
-					Windows::FreeLibrary((Windows::HMODULE)VulkanLibraryPtr);
-					VulkanLibraryPtr = nullptr;
-				}
-			}
+	void FTextureShareCoreVulkanLibrary::Unload()
+	{
+		if (VulkanLibraryPtr)
+		{
+			Windows::FreeLibrary((Windows::HMODULE)VulkanLibraryPtr);
+			VulkanLibraryPtr = nullptr;
+		}
+	}
 
-			static FTextureShareCoreVulkanLibrary& GetVulkanLibrary()
-			{
-				static FTextureShareCoreVulkanLibrary VulkanLibrary;
-				return VulkanLibrary;
-			}
+	static FTextureShareCoreVulkanLibrary& GetVulkanLibrary()
+	{
+		static FTextureShareCoreVulkanLibrary VulkanLibrary;
+		return VulkanLibrary;
+	}
 
-			static bool VulkanCreateSharedHandle(const FTextureShareDeviceVulkanContext& InDeviceVulkanContext, const FTextureShareDeviceVulkanResource& InVulkanResource, const bool bUseNTHandle, HANDLE& OutSharedHandle)
+	static bool VulkanCreateSharedHandle(const FTextureShareDeviceVulkanContext& InDeviceVulkanContext, const FTextureShareDeviceVulkanResource& InVulkanResource, const bool bUseNTHandle, HANDLE& OutSharedHandle)
+	{
+		if (GetVulkanLibrary().IsValid())
+		{
+			if (InVulkanResource.IsValid() && InDeviceVulkanContext.IsValid())
 			{
-				if (GetVulkanLibrary().IsValid())
-				{
-					if (InVulkanResource.IsValid() && InDeviceVulkanContext.IsValid())
-					{
-						// Generate VkMemoryGetWin32HandleInfoKHR
-						VkMemoryGetWin32HandleInfoKHR MemoryGetHandleInfoKHR = {};
+				// Generate VkMemoryGetWin32HandleInfoKHR
+				VkMemoryGetWin32HandleInfoKHR MemoryGetHandleInfoKHR = {};
 
-						MemoryGetHandleInfoKHR.sType = VK_STRUCTURE_TYPE_MEMORY_GET_WIN32_HANDLE_INFO_KHR;
-						MemoryGetHandleInfoKHR.pNext = NULL;
-						MemoryGetHandleInfoKHR.memory = InVulkanResource.GetAllocationHandle();
-						MemoryGetHandleInfoKHR.handleType = bUseNTHandle ? VK_EXTERNAL_MEMORY_HANDLE_TYPE_OPAQUE_WIN32_BIT : VK_EXTERNAL_MEMORY_HANDLE_TYPE_OPAQUE_WIN32_KMT_BIT;
+				MemoryGetHandleInfoKHR.sType = VK_STRUCTURE_TYPE_MEMORY_GET_WIN32_HANDLE_INFO_KHR;
+				MemoryGetHandleInfoKHR.pNext = NULL;
+				MemoryGetHandleInfoKHR.memory = InVulkanResource.GetAllocationHandle();
+				MemoryGetHandleInfoKHR.handleType = bUseNTHandle ? VK_EXTERNAL_MEMORY_HANDLE_TYPE_OPAQUE_WIN32_BIT : VK_EXTERNAL_MEMORY_HANDLE_TYPE_OPAQUE_WIN32_KMT_BIT;
 
-						// While this operation is safe (and unavoidable) C4191 has been enabled and this will trigger an error with warnings as errors
+				// While this operation is safe (and unavoidable) C4191 has been enabled and this will trigger an error with warnings as errors
 #pragma warning(push)
 #pragma warning(disable : 4191)
-						if (GetVulkanDeviceProc(GetMemoryWin32HandleKHR))
-						{
-							const VkResult Result = GetMemoryWin32HandleKHR(InDeviceVulkanContext.GetDevice(), &MemoryGetHandleInfoKHR, &OutSharedHandle);
-							if (Result == VK_SUCCESS)
-							{
-								return OutSharedHandle != 0 && OutSharedHandle != INVALID_HANDLE_VALUE;
-							}
-						}
-#pragma warning(pop)
+				if (GetVulkanDeviceProc(GetMemoryWin32HandleKHR))
+				{
+					const VkResult Result = GetMemoryWin32HandleKHR(InDeviceVulkanContext.GetDevice(), &MemoryGetHandleInfoKHR, &OutSharedHandle);
+					if (Result == VK_SUCCESS)
+					{
+						return OutSharedHandle != 0 && OutSharedHandle != INVALID_HANDLE_VALUE;
 					}
 				}
-
-				return false;
+#pragma warning(pop)
 			}
 		}
+
+		return false;
 	}
 };
 using namespace UE::TextureShareCore::VulkanResourcesCacheItem;
diff --git a/Engine/Plugins/VirtualProduction/TextureShare/Source/TextureShareCore/Private/Core/TextureShareCore.cpp b/Engine/Plugins/VirtualProduction/TextureShare/Source/TextureShareCore/Private/Core/TextureShareCore.cpp
--- a/Engine/Plugins/VirtualProduction/TextureShare/Source/TextureShareCore/Private/Core/TextureShareCore.cpp
+++ b/Engine/Plugins/VirtualProduction/TextureShare/Source/TextureShareCore/Private/Core/TextureShareCore.cpp
@@ -19,23 +19,20 @@
 #include "IPC/TextureShareCoreInterprocessMemoryRegion.h"
 #include "IPC/Containers/TextureShareCoreInterprocessMemory.h"
 
-namespace UE
+namespace UE::TextureShareCore
 {
-	namespace TextureShareCore
+	static inline FString GetValidTextureShareProcessName(const FString& InProcessName)
 	{
-		static inline FString GetValidTextureShareProcessName(const FString& InProcessName)
+		if (InProcessName.IsEmpty())
 		{
-			if (InProcessName.IsEmpty())
-			{
 #if TEXTURESHARECORE_SDK
-				return UE::TextureShareCoreStrings::DefaultProcessName::SDK;
+			return UE::TextureShareCoreStrings::DefaultProcessName::SDK;
 #else
-				return UE::TextureShareCoreStrings::DefaultProcessName::UE;
+			return UE::TextureShareCoreStrings::DefaultProcessName::UE;
 #endif
-			};
+		};
 
-			return InProcessName;
-		}
+		return InProcessName;
 	}
 };
 using namespace UE::TextureShareCore;
diff --git a/Engine/Plugins/VirtualProduction/TextureShare/Source/TextureShareCore/Private/Core/TextureShareCoreHelpers.h b/Engine/Plugins/VirtualProduction/TextureShare/Source/TextureShareCore/Private/Core/TextureShareCoreHelpers.h
--- a/Engine/Plugins/VirtualProduction/TextureShare/Source/TextureShareCore/Private/Core/TextureShareCoreHelpers.h
+++ b/Engine/Plugins/VirtualProduction/TextureShare/Source/TextureShareCore/Private/Core/TextureShareCoreHelpers.h
@@ -9,213 +9,209 @@
 /**
  * Enums to text wrappers for logging
  */
-namespace UE
+namespace UE::TextureShareCore
 {
-	namespace TextureShareCore
+	// Debug log purpose
+	static constexpr auto GetTEXT(const ETextureShareDeviceType In)
 	{
-		// Debug log purpose
-		static constexpr auto GetTEXT(const ETextureShareDeviceType In)
+		switch (In)
 		{
-			switch (In)
-			{
-			case ETextureShareDeviceType::D3D11:  return TEXT("D3D11");
-			case ETextureShareDeviceType::D3D12:  return TEXT("D3D12");
-			case ETextureShareDeviceType::Vulkan: return TEXT("Vulkan");
-			default:
-				break;
-			}
-
-			return TEXT("Undefined");
+		case ETextureShareDeviceType::D3D11:  return TEXT("D3D11");
+		case ETextureShareDeviceType::D3D12:  return TEXT("D3D12");
+		case ETextureShareDeviceType::Vulkan: return TEXT("Vulkan");
+		default:
+			break;
 		}
 
-		static constexpr auto GetTEXT(const ETextureShareFrameSyncTemplate In)
-		{
-			switch (In)
-			{
-			case ETextureShareFrameSyncTemplate::Default: return TEXT("Default");
-			case ETextureShareFrameSyncTemplate::SDK:     return TEXT("SDK");
-			case ETextureShareFrameSyncTemplate::DisplayCluster: return TEXT("DisplayCluster");
-			default:
-				break;
-			}
+		return TEXT("Undefined");
+	}
 
-			return TEXT("?");
+	static constexpr auto GetTEXT(const ETextureShareFrameSyncTemplate In)
+	{
+		switch (In)
+		{
+		case ETextureShareFrameSyncTemplate::Default: return TEXT("Default");
+		case ETextureShareFrameSyncTemplate::SDK:     return TEXT("SDK");
+		case ETextureShareFrameSyncTemplate::DisplayCluster: return TEXT("DisplayCluster");
+		default:
+			break;
 		}
 
-		static constexpr auto GetTEXT(const ETextureShareThreadMutex In)
-		{
-			switch (In)
-			{
-			case ETextureShareThreadMutex::GameThread:      return TEXT("GameThread");
-			case ETextureShareThreadMutex::RenderingThread: return TEXT("RenderingThread");
-			case ETextureShareThreadMutex::InternalLock:    return TEXT("InternalLock");
-			default:
-				break;
-			}
+		return TEXT("?");
+	}
 
-			return TEXT("?");
+	static constexpr auto GetTEXT(const ETextureShareThreadMutex In)
+	{
+		switch (In)
+		{
+		case ETextureShareThreadMutex::GameThread:      return TEXT("GameThread");
+		case ETextureShareThreadMutex::RenderingThread: return TEXT("RenderingThread");
+		case ETextureShareThreadMutex::InternalLock:    return TEXT("InternalLock");
+		default:
+			break;
 		}
 
-		static constexpr auto GetTEXT(const ETextureShareTextureOp In)
-		{
-			switch (In)
-			{
-			case ETextureShareTextureOp::Write:    return TEXT("Write");
-			case ETextureShareTextureOp::Read:     return TEXT("Read");
-			default:
-				break;
-			}
+		return TEXT("?");
+	}
 
-			return TEXT("?");
+	static constexpr auto GetTEXT(const ETextureShareTextureOp In)
+	{
+		switch (In)
+		{
+		case ETextureShareTextureOp::Write:    return TEXT("Write");
+		case ETextureShareTextureOp::Read:     return TEXT("Read");
+		default:
+			break;
 		}
 
-		static constexpr auto GetTEXT(const ETextureShareEyeType InEyeType)
-		{
-			switch (InEyeType)
-			{
-			case ETextureShareEyeType::Default:     return TEXT("Default");
-			case ETextureShareEyeType::StereoLeft:  return TEXT("StereoLeft");
-			case ETextureShareEyeType::StereoRight: return TEXT("StereoRight");
-			default:
-				break;
-			}
+		return TEXT("?");
+	}
 
-			return TEXT("?");
+	static constexpr auto GetTEXT(const ETextureShareEyeType InEyeType)
+	{
+		switch (InEyeType)
+		{
+		case ETextureShareEyeType::Default:     return TEXT("Default");
+		case ETextureShareEyeType::StereoLeft:  return TEXT("StereoLeft");
+		case ETextureShareEyeType::StereoRight: return TEXT("StereoRight");
+		default:
+			break;
 		}
 
-		static constexpr auto GetTEXT(const ETextureShareSyncStep In)
-		{
-			switch (In)
-			{
-			case ETextureShareSyncStep::InterprocessConnection: return TEXT("IC");
+		return TEXT("?");
+	}
 
-				/**
-				 * Frame sync steps(GameThread)
-				 */
+	static constexpr auto GetTEXT(const ETextureShareSyncStep In)
+	{
+		switch (In)
+		{
+		case ETextureShareSyncStep::InterprocessConnection: return TEXT("IC");
 
-			case ETextureShareSyncStep::FrameBegin: return TEXT("f");
+			/**
+				* Frame sync steps(GameThread)
+				*/
 
-			case ETextureShareSyncStep::FramePreSetupBegin: return TEXT("fpS");
-			case ETextureShareSyncStep::FramePreSetupEnd: return TEXT("fpS+");
+		case ETextureShareSyncStep::FrameBegin: return TEXT("f");
 
-			case ETextureShareSyncStep::FrameSetupBegin: return TEXT("fS");
-			case ETextureShareSyncStep::FrameSetupEnd: return TEXT("fS+");
+		case ETextureShareSyncStep::FramePreSetupBegin: return TEXT("fpS");
+		case ETextureShareSyncStep::FramePreSetupEnd: return TEXT("fpS+");
 
-			case ETextureShareSyncStep::FramePostSetupBegin: return TEXT("fPS");
-			case ETextureShareSyncStep::FramePostSetupEnd: return TEXT("fPS+");
+		case ETextureShareSyncStep::FrameSetupBegin: return TEXT("fS");
+		case ETextureShareSyncStep::FrameSetupEnd: return TEXT("fS+");
 
-			case ETextureShareSyncStep::FrameFlush: return TEXT("flush");
-			case ETextureShareSyncStep::FrameEnd: return TEXT("f+");
+		case ETextureShareSyncStep::FramePostSetupBegin: return TEXT("fPS");
+		case ETextureShareSyncStep::FramePostSetupEnd: return TEXT("fPS+");
 
-				/**
-				 * FrameProxy sync steps (RenderThread)
-				 */
+		case ETextureShareSyncStep::FrameFlush: return TEXT("flush");
+		case ETextureShareSyncStep::FrameEnd: return TEXT("f+");
 
-			case ETextureShareSyncStep::FrameProxyBegin: return TEXT("F");
+			/**
+				* FrameProxy sync steps (RenderThread)
+				*/
 
-			case ETextureShareSyncStep::FrameSceneFinalColorBegin: return TEXT("F_fc");
-			case ETextureShareSyncStep::FrameSceneFinalColorEnd: return TEXT("F_fc+");
+		case ETextureShareSyncStep::FrameProxyBegin: return TEXT("F");
 
-			case ETextureShareSyncStep::FrameProxyPreRenderBegin: return TEXT("FpR");
-			case ETextureShareSyncStep::FrameProxyPreRenderEnd: return TEXT("FpR+");
+		case ETextureShareSyncStep::FrameSceneFinalColorBegin: return TEXT("F_fc");
+		case ETextureShareSyncStep::FrameSceneFinalColorEnd: return TEXT("F_fc+");
 
-			case ETextureShareSyncStep::FrameProxyRenderBegin: return TEXT("FR");
-			case ETextureShareSyncStep::FrameProxyRenderEnd: return TEXT("FR+");
+		case ETextureShareSyncStep::FrameProxyPreRenderBegin: return TEXT("FpR");
+		case ETextureShareSyncStep::FrameProxyPreRenderEnd: return TEXT("FpR+");
 
-			case ETextureShareSyncStep::FrameProxyPostWarpBegin: return TEXT("FPW");
-			case ETextureShareSyncStep::FrameProxyPostWarpEnd: return TEXT("FPW+");
+		case ETextureShareSyncStep::FrameProxyRenderBegin: return TEXT("FR");
+		case ETextureShareSyncStep::FrameProxyRenderEnd: return TEXT("FR+");
 
-			case ETextureShareSyncStep::FrameProxyPostRenderBegin: return TEXT("FPR");
-			case ETextureShareSyncStep::FrameProxyPostRenderEnd: return TEXT("FPR+");
+		case ETextureShareSyncStep::FrameProxyPostWarpBegin: return TEXT("FPW");
+		case ETextureShareSyncStep::FrameProxyPostWarpEnd: return TEXT("FPW+");
 
-			case ETextureShareSyncStep::FrameProxyBackBufferReadyToPresentBegin: return TEXT("FBbPr");
-			case ETextureShareSyncStep::FrameProxyBackBufferReadyToPresentEnd: return TEXT("FBbPr+");
+		case ETextureShareSyncStep::FrameProxyPostRenderBegin: return TEXT("FPR");
+		case ETextureShareSyncStep::FrameProxyPostRenderEnd: return TEXT("FPR+");
 
-			case ETextureShareSyncStep::FrameProxyFlush: return TEXT("FLUSH");
-			case ETextureShareSyncStep::FrameProxyEnd: return TEXT("F+");
+		case ETextureShareSyncStep::FrameProxyBackBufferReadyToPresentBegin: return TEXT("FBbPr");
+		case ETextureShareSyncStep::FrameProxyBackBufferReadyToPresentEnd: return TEXT("FBbPr+");
 
-			default:
-				break;
-			}
+		case ETextureShareSyncStep::FrameProxyFlush: return TEXT("FLUSH");
+		case ETextureShareSyncStep::FrameProxyEnd: return TEXT("F+");
 
-			return TEXT("?");
+		default:
+			break;
 		}
 
-		static constexpr auto GetTEXT(const ETextureShareSyncPass In)
-		{
-			switch (In)
-			{
-			case ETextureShareSyncPass::Enter:    return TEXT("Enter");
-			case ETextureShareSyncPass::Exit:     return TEXT("Exit");
-			case ETextureShareSyncPass::Complete: return TEXT("Complete");
-			default:
-				break;
-			}
+		return TEXT("?");
+	}
 
-			return TEXT("?");
+	static constexpr auto GetTEXT(const ETextureShareSyncPass In)
+	{
+		switch (In)
+		{
+		case ETextureShareSyncPass::Enter:    return TEXT("Enter");
+		case ETextureShareSyncPass::Exit:     return TEXT("Exit");
+		case ETextureShareSyncPass::Complete: return TEXT("Complete");
+		default:
+			break;
 		}
 
-		static constexpr auto GetTEXT(const ETextureShareSyncState In)
-		{
-			switch (In)
-			{
-			case ETextureShareSyncState::Enter:          return TEXT("1");
-			case ETextureShareSyncState::EnterCompleted: return TEXT("1+");
+		return TEXT("?");
+	}
 
-			case ETextureShareSyncState::Exit:          return TEXT("2");
-			case ETextureShareSyncState::ExitCompleted: return TEXT("2+");
+	static constexpr auto GetTEXT(const ETextureShareSyncState In)
+	{
+		switch (In)
+		{
+		case ETextureShareSyncState::Enter:          return TEXT("1");
+		case ETextureShareSyncState::EnterCompleted: return TEXT("1+");
 
-			case ETextureShareSyncState::Completed: return TEXT("##");
+		case ETextureShareSyncState::Exit:          return TEXT("2");
+		case ETextureShareSyncState::ExitCompleted: return TEXT("2+");
 
-			case ETextureShareSyncState::Undefined:
-				break;
-			}
+		case ETextureShareSyncState::Completed: return TEXT("##");
 
-			return TEXT("?");
+		case ETextureShareSyncState::Undefined:
+			break;
 		}
 
-		static FString ToString(const FTextureShareCoreObjectDesc& InObject)
-		{
-			const FTextureShareCoreObjectSyncState& InSyncState = InObject.Sync.SyncState;
-			const FString SyncState = FString::Printf(TEXT("%s[%s]{%s < %s}"), GetTEXT(InSyncState.Step), GetTEXT(InSyncState.State), GetTEXT(InSyncState.NextStep), GetTEXT(InSyncState.PrevStep));
+		return TEXT("?");
+	}
 
-			return FString::Printf(TEXT("%s(%s)"), *InObject.ProcessDesc.ProcessId, *SyncState);
-		}
+	static FString ToString(const FTextureShareCoreObjectDesc& InObject)
+	{
+		const FTextureShareCoreObjectSyncState& InSyncState = InObject.Sync.SyncState;
+		const FString SyncState = FString::Printf(TEXT("%s[%s]{%s < %s}"), GetTEXT(InSyncState.Step), GetTEXT(InSyncState.State), GetTEXT(InSyncState.NextStep), GetTEXT(InSyncState.PrevStep));
 
-		static FString ToString(const TArraySerializable<FTextureShareCoreObjectDesc>& InObjects)
+		return FString::Printf(TEXT("%s(%s)"), *InObject.ProcessDesc.ProcessId, *SyncState);
+	}
+
+	static FString ToString(const TArraySerializable<FTextureShareCoreObjectDesc>& InObjects)
+	{
+		FString Result;
+		for (const FTextureShareCoreObjectDesc& ObjectDescIt : InObjects)
 		{
-			FString Result;
-			for (const FTextureShareCoreObjectDesc& ObjectDescIt : InObjects)
+			if (!Result.IsEmpty())
 			{
-				if (!Result.IsEmpty())
-				{
-					Result += TEXT(", ");
-				}
-
-				Result += ToString(ObjectDescIt);
+				Result += TEXT(", ");
 			}
 
-			return Result;
+			Result += ToString(ObjectDescIt);
 		}
 
-		static FString ToString(const FTextureShareCoreViewDesc& InViewDesc)
-		{
-			return FString::Printf(TEXT("%s[%s]"), *InViewDesc.Id, GetTEXT(InViewDesc.EyeType));
-		}
+		return Result;
+	}
 
-		static FString ToString(const FTextureShareCoreResourceDesc& InResourceDesc)
-		{
-			return FString::Printf(TEXT("%s.%s, Eye=%s, Op=%s, SyncStep=%s, SrcViewId='%s'"), *InResourceDesc.ViewDesc.Id, *InResourceDesc.ResourceName,
-				GetTEXT(InResourceDesc.ViewDesc.EyeType),
-				GetTEXT(InResourceDesc.OperationType), GetTEXT(InResourceDesc.SyncStep),
-				*InResourceDesc.ViewDesc.SrcId
-			);
-		}
+	static FString ToString(const FTextureShareCoreViewDesc& InViewDesc)
+	{
+		return FString::Printf(TEXT("%s[%s]"), *InViewDesc.Id, GetTEXT(InViewDesc.EyeType));
+	}
 
-		static FString ToString(const FTextureShareCoreResourceRequest& InResourceRequest)
-		{
-			return FString::Printf(TEXT("%s, GPU#=%d"), *ToString(InResourceRequest.ResourceDesc), InResourceRequest.GPUIndex);
-		}
+	static FString ToString(const FTextureShareCoreResourceDesc& InResourceDesc)
+	{
+		return FString::Printf(TEXT("%s.%s, Eye=%s, Op=%s, SyncStep=%s, SrcViewId='%s'"), *InResourceDesc.ViewDesc.Id, *InResourceDesc.ResourceName,
+			GetTEXT(InResourceDesc.ViewDesc.EyeType),
+			GetTEXT(InResourceDesc.OperationType), GetTEXT(InResourceDesc.SyncStep),
+			*InResourceDesc.ViewDesc.SrcId
+		);
+	}
 
+	static FString ToString(const FTextureShareCoreResourceRequest& InResourceRequest)
+	{
+		return FString::Printf(TEXT("%s, GPU#=%d"), *ToString(InResourceRequest.ResourceDesc), InResourceRequest.GPUIndex);
 	}
 };
diff --git a/Engine/Plugins/VirtualProduction/TextureShare/Source/TextureShareCore/Private/IPC/Containers/TextureShareCoreInterprocessContainers.cpp b/Engine/Plugins/VirtualProduction/TextureShare/Source/TextureShareCore/Private/IPC/Containers/TextureShareCoreInterprocessContainers.cpp
--- a/Engine/Plugins/VirtualProduction/TextureShare/Source/TextureShareCore/Private/IPC/Containers/TextureShareCoreInterprocessContainers.cpp
+++ b/Engine/Plugins/VirtualProduction/TextureShare/Source/TextureShareCore/Private/IPC/Containers/TextureShareCoreInterprocessContainers.cpp
@@ -5,20 +5,14 @@
 #include "Misc/SecureHash.h"
 #include "Windows/MinimalWindowsApi.h"
 
-namespace UE
+namespace UE::TextureShareCore::InterprocessContainers
 {
-	namespace TextureShareCore
+	static FTextureShareCoreSMD5Hash CreateEmptySMD5Hash()
 	{
-		namespace InterprocessContainers
-		{
-			static FTextureShareCoreSMD5Hash CreateEmptySMD5Hash()
-			{
-				FTextureShareCoreSMD5Hash Result;
-				Result.Empty();
+		FTextureShareCoreSMD5Hash Result;
+		Result.Empty();
 
-				return Result;
-			}
-		}
+		return Result;
 	}
 };
 using namespace UE::TextureShareCore::InterprocessContainers;
diff --git a/Engine/Plugins/VirtualProduction/TextureShare/Source/TextureShareCore/Private/IPC/TextureShareCoreInterprocessEvent.cpp b/Engine/Plugins/VirtualProduction/TextureShare/Source/TextureShareCore/Private/IPC/TextureShareCoreInterprocessEvent.cpp
--- a/Engine/Plugins/VirtualProduction/TextureShare/Source/TextureShareCore/Private/IPC/TextureShareCoreInterprocessEvent.cpp
+++ b/Engine/Plugins/VirtualProduction/TextureShare/Source/TextureShareCore/Private/IPC/TextureShareCoreInterprocessEvent.cpp
@@ -11,64 +11,58 @@
 #include "Windows/WindowsHWrapper.h"
 #include "Windows/AllowWindowsPlatformTypes.h"
 
-namespace UE
+namespace UE::TextureShareCore::InterprocessEvent
 {
-	namespace TextureShareCore
+	static TSharedPtr<FEvent, ESPMode::ThreadSafe> CreateNamedInterprocessEvent(const FString& EventName, const bool bOpenExistEvent, const void* InSecurityAttributes)
 	{
-		namespace InterprocessEvent
-		{
-			static TSharedPtr<FEvent, ESPMode::ThreadSafe> CreateNamedInterprocessEvent(const FString& EventName, const bool bOpenExistEvent, const void* InSecurityAttributes)
-			{
-				const FString GlobalEventName = FString::Printf(TEXT("Global\\%s"), *EventName);
-
-				HANDLE EventHandle = NULL;
-
-				if (bOpenExistEvent)
-				{
-					const DWORD AccessRights = EVENT_ALL_ACCESS | EVENT_MODIFY_STATE;//SYNCHRONIZE | EVENT_MODIFY_STATE;
-					EventHandle = OpenEventA(AccessRights, false, TCHAR_TO_ANSI(*GlobalEventName));
-					if (NULL == EventHandle)
-					{
-						UE_LOG(LogTextureShareCore, Warning, TEXT("OpenEvent(AccessRights=0x%08x, bInherit=false, Name='%s') failed with LastError = %d"), AccessRights, *EventName, GetLastError());
-
-						return nullptr;
-					}
-				}
-				else
-				{
-					// Create new event
-					EventHandle = CreateEventA((SECURITY_ATTRIBUTES*)InSecurityAttributes, true, false, TCHAR_TO_ANSI(*GlobalEventName));
-					if (NULL == EventHandle)
-					{
-						UE_LOG(LogTextureShareCore, Warning, TEXT("CreateEvent(NULL, bInherit=false, Name='%s') failed with LastError = %d"), *EventName, GetLastError());
-
-						return nullptr;
-					}
-				}
-
-				return MakeShared<FTextureShareCoreInterprocessEventWin, ESPMode::ThreadSafe>(EventHandle);
-			}
+		const FString GlobalEventName = FString::Printf(TEXT("Global\\%s"), *EventName);
 
-			static DWORD WaitForSingleInterprocessEvent(const HANDLE EventHandle, DWORD dwMilliseconds)
-			{
-				return WaitForSingleObject(EventHandle, dwMilliseconds);
-			}
+		HANDLE EventHandle = NULL;
 
-			static bool SetInterprocessEvent(const HANDLE EventHandle)
+		if (bOpenExistEvent)
+		{
+			const DWORD AccessRights = EVENT_ALL_ACCESS | EVENT_MODIFY_STATE;//SYNCHRONIZE | EVENT_MODIFY_STATE;
+			EventHandle = OpenEventA(AccessRights, false, TCHAR_TO_ANSI(*GlobalEventName));
+			if (NULL == EventHandle)
 			{
-				return SetEvent(EventHandle);
-			}
+				UE_LOG(LogTextureShareCore, Warning, TEXT("OpenEvent(AccessRights=0x%08x, bInherit=false, Name='%s') failed with LastError = %d"), AccessRights, *EventName, GetLastError());
 
-			static bool ResetInterprocessEvent(const HANDLE EventHandle)
-			{
-				return ResetEvent(EventHandle);
+				return nullptr;
 			}
-
-			static bool CloseInterprocessEvent(const HANDLE EventHandle)
+		}
+		else
+		{
+			// Create new event
+			EventHandle = CreateEventA((SECURITY_ATTRIBUTES*)InSecurityAttributes, true, false, TCHAR_TO_ANSI(*GlobalEventName));
+			if (NULL == EventHandle)
 			{
-				return CloseHandle(EventHandle);
+				UE_LOG(LogTextureShareCore, Warning, TEXT("CreateEvent(NULL, bInherit=false, Name='%s') failed with LastError = %d"), *EventName, GetLastError());
+
+				return nullptr;
 			}
 		}
+
+		return MakeShared<FTextureShareCoreInterprocessEventWin, ESPMode::ThreadSafe>(EventHandle);
+	}
+
+	static DWORD WaitForSingleInterprocessEvent(const HANDLE EventHandle, DWORD dwMilliseconds)
+	{
+		return WaitForSingleObject(EventHandle, dwMilliseconds);
+	}
+
+	static bool SetInterprocessEvent(const HANDLE EventHandle)
+	{
+		return SetEvent(EventHandle);
+	}
+
+	static bool ResetInterprocessEvent(const HANDLE EventHandle)
+	{
+		return ResetEvent(EventHandle);
+	}
+
+	static bool CloseInterprocessEvent(const HANDLE EventHandle)
+	{
+		return CloseHandle(EventHandle);
 	}
 };
 using namespace UE::TextureShareCore::InterprocessEvent;
diff --git a/Engine/Plugins/VirtualProduction/TextureShare/Source/TextureShareCore/Private/IPC/TextureShareCoreInterprocessHelpers.h b/Engine/Plugins/VirtualProduction/TextureShare/Source/TextureShareCore/Private/IPC/TextureShareCoreInterprocessHelpers.h
--- a/Engine/Plugins/VirtualProduction/TextureShare/Source/TextureShareCore/Private/IPC/TextureShareCoreInterprocessHelpers.h
+++ b/Engine/Plugins/VirtualProduction/TextureShare/Source/TextureShareCore/Private/IPC/TextureShareCoreInterprocessHelpers.h
@@ -6,77 +6,74 @@
 /**
  * Enums to text wrappers for logging
  */
-namespace UE
+namespace UE::TextureShareCore
 {
-	namespace TextureShareCore
+	static constexpr auto GetTEXT(const ETextureShareInterprocessObjectSyncBarrierState InBarrierState)
 	{
-		static constexpr auto GetTEXT(const ETextureShareInterprocessObjectSyncBarrierState InBarrierState)
+		switch (InBarrierState)
 		{
-			switch (InBarrierState)
-			{
-			case ETextureShareInterprocessObjectSyncBarrierState::Undefined: return TEXT("Undefined");
-			case ETextureShareInterprocessObjectSyncBarrierState::UnusedSyncStep: return TEXT("UnusedSyncStep");
+		case ETextureShareInterprocessObjectSyncBarrierState::Undefined: return TEXT("Undefined");
+		case ETextureShareInterprocessObjectSyncBarrierState::UnusedSyncStep: return TEXT("UnusedSyncStep");
 
-			case ETextureShareInterprocessObjectSyncBarrierState::WaitConnection: return TEXT("WaitConnection");
-			case ETextureShareInterprocessObjectSyncBarrierState::AcceptConnection: return TEXT("AcceptConnection");
+		case ETextureShareInterprocessObjectSyncBarrierState::WaitConnection: return TEXT("WaitConnection");
+		case ETextureShareInterprocessObjectSyncBarrierState::AcceptConnection: return TEXT("AcceptConnection");
 
-			case ETextureShareInterprocessObjectSyncBarrierState::Wait: return TEXT("Wait");
-			case ETextureShareInterprocessObjectSyncBarrierState::Accept: return TEXT("Accept");
+		case ETextureShareInterprocessObjectSyncBarrierState::Wait: return TEXT("Wait");
+		case ETextureShareInterprocessObjectSyncBarrierState::Accept: return TEXT("Accept");
 
-			case ETextureShareInterprocessObjectSyncBarrierState::ObjectLost: return TEXT("ObjectLost");
-			case ETextureShareInterprocessObjectSyncBarrierState::FrameSyncLost: return TEXT("FrameSyncLost");
+		case ETextureShareInterprocessObjectSyncBarrierState::ObjectLost: return TEXT("ObjectLost");
+		case ETextureShareInterprocessObjectSyncBarrierState::FrameSyncLost: return TEXT("FrameSyncLost");
 
-			case ETextureShareInterprocessObjectSyncBarrierState::InvalidLogic: return TEXT("InvalidLogic");
-			}
+		case ETextureShareInterprocessObjectSyncBarrierState::InvalidLogic: return TEXT("InvalidLogic");
+		}
 
-			return TEXT("Unknown");
-		};
+		return TEXT("Unknown");
+	};
 
-		static constexpr auto GetTEXT(const ETextureShareCoreInterprocessObjectFrameSyncState InSyncState)
+	static constexpr auto GetTEXT(const ETextureShareCoreInterprocessObjectFrameSyncState InSyncState)
+	{
+		switch (InSyncState)
 		{
-			switch (InSyncState)
-			{
-			case  ETextureShareCoreInterprocessObjectFrameSyncState::Undefined: return TEXT("Undefined");
-			case  ETextureShareCoreInterprocessObjectFrameSyncState::FrameSyncLost: return TEXT("FrameSyncLost");
-
-			case  ETextureShareCoreInterprocessObjectFrameSyncState::NewFrame: return TEXT("NewFrame");
-			case  ETextureShareCoreInterprocessObjectFrameSyncState::FrameConnected: return TEXT("FrameConnected");
-			case  ETextureShareCoreInterprocessObjectFrameSyncState::FrameBegin: return TEXT("FrameBegin");
-			case  ETextureShareCoreInterprocessObjectFrameSyncState::FrameEnd: return TEXT("FrameEnd");
-
-			case  ETextureShareCoreInterprocessObjectFrameSyncState::FrameProxyBegin:  return TEXT("FrameProxyBegin");
-			case  ETextureShareCoreInterprocessObjectFrameSyncState::FrameProxyEnd:  return TEXT("FrameProxyEnd");
-			default:
-				break;
-			}
+		case  ETextureShareCoreInterprocessObjectFrameSyncState::Undefined: return TEXT("Undefined");
+		case  ETextureShareCoreInterprocessObjectFrameSyncState::FrameSyncLost: return TEXT("FrameSyncLost");
+
+		case  ETextureShareCoreInterprocessObjectFrameSyncState::NewFrame: return TEXT("NewFrame");
+		case  ETextureShareCoreInterprocessObjectFrameSyncState::FrameConnected: return TEXT("FrameConnected");
+		case  ETextureShareCoreInterprocessObjectFrameSyncState::FrameBegin: return TEXT("FrameBegin");
+		case  ETextureShareCoreInterprocessObjectFrameSyncState::FrameEnd: return TEXT("FrameEnd");
+
+		case  ETextureShareCoreInterprocessObjectFrameSyncState::FrameProxyBegin:  return TEXT("FrameProxyBegin");
+		case  ETextureShareCoreInterprocessObjectFrameSyncState::FrameProxyEnd:  return TEXT("FrameProxyEnd");
+		default:
+			break;
+		}
 
-			return TEXT("Unknown");
-		};
+		return TEXT("Unknown");
+	};
 
-		static FString ToString(const FTextureShareCoreInterprocessObjectSyncState& InObjectSyncState)
-		{
-			return FString::Printf(TEXT("%s[%s]{%s < %s}"), GetTEXT(InObjectSyncState.Step), GetTEXT(InObjectSyncState.State), GetTEXT(InObjectSyncState.NextStep), GetTEXT(InObjectSyncState.PrevStep));
-		}
+	static FString ToString(const FTextureShareCoreInterprocessObjectSyncState& InObjectSyncState)
+	{
+		return FString::Printf(TEXT("%s[%s]{%s < %s}"), GetTEXT(InObjectSyncState.Step), GetTEXT(InObjectSyncState.State), GetTEXT(InObjectSyncState.NextStep), GetTEXT(InObjectSyncState.PrevStep));
+	}
 
-		static FString ToString(const FTextureShareCoreInterprocessObject& InObject)
-		{
-			return FString::Printf(TEXT("%s(%s)"), *InObject.Desc.ProcessName.ToString(), *ToString(InObject.Sync.GetSyncState()));
-		}
+	static FString ToString(const FTextureShareCoreInterprocessObject& InObject)
+	{
+		return FString::Printf(TEXT("%s(%s)"), *InObject.Desc.ProcessName.ToString(), *ToString(InObject.Sync.GetSyncState()));
+	}
 
-		static FString ToString(const TArray<FTextureShareCoreInterprocessObject*>& InObjects)
+	static FString ToString(const TArray<FTextureShareCoreInterprocessObject*>& InObjects)
+	{
+		FString Result;
+		for (const FTextureShareCoreInterprocessObject* It : InObjects)
 		{
-			FString Result;
-			for (const FTextureShareCoreInterprocessObject* It : InObjects)
+			if (!Result.IsEmpty())
 			{
-				if (!Result.IsEmpty())
-				{
-					Result += TEXT(", ");
-				}
-
-				Result += ToString(*It);
+				Result += TEXT(", ");
 			}
 
-			return Result;
+			Result += ToString(*It);
 		}
+
+		return Result;
 	}
 };
diff --git a/Engine/Plugins/VirtualProduction/TextureShare/Source/TextureShareCore/Private/Misc/TextureShareCoreStrings.h b/Engine/Plugins/VirtualProduction/TextureShare/Source/TextureShareCore/Private/Misc/TextureShareCoreStrings.h
--- a/Engine/Plugins/VirtualProduction/TextureShare/Source/TextureShareCore/Private/Misc/TextureShareCoreStrings.h
+++ b/Engine/Plugins/VirtualProduction/TextureShare/Source/TextureShareCore/Private/Misc/TextureShareCoreStrings.h
@@ -10,23 +10,20 @@
 /**
  * Default names for TextureShareCore module
  */
-namespace UE
+namespace UE::TextureShareCoreStrings
 {
-	namespace TextureShareCoreStrings
-	{
-		static constexpr auto DefaultShareName = TEXT("DefaultShareName");
+	static constexpr auto DefaultShareName = TEXT("DefaultShareName");
 
-		static constexpr auto DefaultViewId = TEXT("DefaultView");
+	static constexpr auto DefaultViewId = TEXT("DefaultView");
 
-		namespace DefaultProcessName
-		{
+	namespace DefaultProcessName
+	{
 #if defined TEXTURESHARE_PROJECT_NAME
-			// Use TextureShare project name as process name
-			static constexpr auto SDK = TEXT(TEXTURESHARE_PROJECT_NAME);
+		// Use TextureShare project name as process name
+		static constexpr auto SDK = TEXT(TEXTURESHARE_PROJECT_NAME);
 #else
-			static constexpr auto SDK = TEXT("TextureShareSDK");
+		static constexpr auto SDK = TEXT("TextureShareSDK");
 #endif
-			static constexpr auto UE = TEXT("UnrealEngine");
-		}
+		static constexpr auto UE = TEXT("UnrealEngine");
 	}
 };
diff --git a/Engine/Plugins/VirtualProduction/TextureShare/Source/TextureShareCore/Private/Object/TextureShareCoreObject.cpp b/Engine/Plugins/VirtualProduction/TextureShare/Source/TextureShareCore/Private/Object/TextureShareCoreObject.cpp
--- a/Engine/Plugins/VirtualProduction/TextureShare/Source/TextureShareCore/Private/Object/TextureShareCoreObject.cpp
+++ b/Engine/Plugins/VirtualProduction/TextureShare/Source/TextureShareCore/Private/Object/TextureShareCoreObject.cpp
@@ -14,31 +14,28 @@
 
 #include "Misc/ScopeLock.h"
 
-namespace UE
+namespace UE::TextureShareCore
 {
-	namespace TextureShareCore
+	static FTextureShareCoreObjectDesc CreateNewObjectDesc(const FTextureShareCore& InOwner, const FString& InTextureShareName, const ETextureShareProcessType InProcessType)
 	{
-		static FTextureShareCoreObjectDesc CreateNewObjectDesc(const FTextureShareCore& InOwner, const FString& InTextureShareName, const ETextureShareProcessType InProcessType)
-		{
-			FTextureShareCoreObjectDesc ObjectDesc;
-
-			// The name of this object
-			ObjectDesc.ShareName = InTextureShareName;
+		FTextureShareCoreObjectDesc ObjectDesc;
 
-			// Generate unique Guid for each object
-			ObjectDesc.ObjectGuid = FGuid::NewGuid();
+		// The name of this object
+		ObjectDesc.ShareName = InTextureShareName;
 
-			// Get value for Owner
-			ObjectDesc.ProcessDesc = InOwner.GetProcessDesc();
+		// Generate unique Guid for each object
+		ObjectDesc.ObjectGuid = FGuid::NewGuid();
 
-			// Custom process type
-			if (InProcessType != ETextureShareProcessType::Undefined)
-			{
-				ObjectDesc.ProcessDesc.ProcessType = InProcessType;
-			}
+		// Get value for Owner
+		ObjectDesc.ProcessDesc = InOwner.GetProcessDesc();
 
-			return ObjectDesc;
+		// Custom process type
+		if (InProcessType != ETextureShareProcessType::Undefined)
+		{
+			ObjectDesc.ProcessDesc.ProcessType = InProcessType;
 		}
+
+		return ObjectDesc;
 	}
 };
 using namespace UE::TextureShareCore;
diff --git a/Engine/Plugins/VirtualProduction/TextureShare/Source/TextureShareDisplayCluster/Private/Misc/TextureShareDisplayClusterStrings.h b/Engine/Plugins/VirtualProduction/TextureShare/Source/TextureShareDisplayCluster/Private/Misc/TextureShareDisplayClusterStrings.h
--- a/Engine/Plugins/VirtualProduction/TextureShare/Source/TextureShareDisplayCluster/Private/Misc/TextureShareDisplayClusterStrings.h
+++ b/Engine/Plugins/VirtualProduction/TextureShare/Source/TextureShareDisplayCluster/Private/Misc/TextureShareDisplayClusterStrings.h
@@ -10,44 +10,38 @@
 /**
  * Constant names for nDisplay texture resources
  */
-namespace UE
+namespace UE::TextureShare::DisplayClusterStrings
 {
-	namespace TextureShare
+	// Default ShareName for nDisplay integration object
+	static constexpr auto DefaultShareName = TEXT("nDisplay");
+
+	namespace Viewport
+	{
+		// FinalColor, but can be overrided
+		static constexpr auto FinalColor = TEXT("ViewportFinalColor");
+
+		// internal resolved shader resources, used as warp&blend source
+		static constexpr auto Input = TEXT("ViewportInput");
+		static constexpr auto Mips = TEXT("ViewportMips");
+
+		// After warp viewport (before output remap)
+		static constexpr auto Warped = TEXT("ViewportWarpBlend");
+	}
+
+	namespace Output
+	{
+		// access to nDisplay frame backbuffer
+		static constexpr auto Backbuffer = TEXT("FrameBackbuffer");
+		static constexpr auto BackbufferTemp = TEXT("FrameBackbufferTemp");
+	}
+
+	namespace Postprocess
+	{
+		static constexpr auto TextureShare = TEXT("TextureShare");
+	}
+
+	namespace Projection
 	{
-		namespace DisplayClusterStrings
-		{
-			// Default ShareName for nDisplay integration object
-			static constexpr auto DefaultShareName = TEXT("nDisplay");
-
-			namespace Viewport
-			{
-				// FinalColor, but can be overrided
-				static constexpr auto FinalColor = TEXT("ViewportFinalColor");
-
-				// internal resolved shader resources, used as warp&blend source
-				static constexpr auto Input = TEXT("ViewportInput");
-				static constexpr auto Mips = TEXT("ViewportMips");
-
-				// After warp viewport (before output remap)
-				static constexpr auto Warped = TEXT("ViewportWarpBlend");
-			}
-
-			namespace Output
-			{
-				// access to nDisplay frame backbuffer
-				static constexpr auto Backbuffer = TEXT("FrameBackbuffer");
-				static constexpr auto BackbufferTemp = TEXT("FrameBackbufferTemp");
-			}
-
-			namespace Postprocess
-			{
-				static constexpr auto TextureShare = TEXT("TextureShare");
-			}
-
-			namespace Projection
-			{
-				static constexpr auto TextureShare = TEXT("textureshare");
-			}
-		}
+		static constexpr auto TextureShare = TEXT("textureshare");
 	}
 };
diff --git a/Engine/Plugins/VirtualProduction/TextureShare/Source/TextureShareDisplayCluster/Private/Module/TextureShareDisplayClusterAPI.cpp b/Engine/Plugins/VirtualProduction/TextureShare/Source/TextureShareDisplayCluster/Private/Module/TextureShareDisplayClusterAPI.cpp
--- a/Engine/Plugins/VirtualProduction/TextureShare/Source/TextureShareDisplayCluster/Private/Module/TextureShareDisplayClusterAPI.cpp
+++ b/Engine/Plugins/VirtualProduction/TextureShare/Source/TextureShareDisplayCluster/Private/Module/TextureShareDisplayClusterAPI.cpp
@@ -13,18 +13,12 @@
 #include "Render/Projection/IDisplayClusterProjectionPolicy.h"
 #include "Render/IDisplayClusterRenderManager.h"
 
-namespace UE
+namespace UE::TextureShare::DisplayCluster
 {
-	namespace TextureShare
+	static IDisplayCluster& DisplayClusterAPI()
 	{
-		namespace DisplayCluster
-		{
-			static IDisplayCluster& DisplayClusterAPI()
-			{
-				static IDisplayCluster& DisplayClusterSingleton = IDisplayCluster::Get();
-				return DisplayClusterSingleton;
-			}
-		}
+		static IDisplayCluster& DisplayClusterSingleton = IDisplayCluster::Get();
+		return DisplayClusterSingleton;
 	}
 };
 using namespace UE::TextureShare::DisplayCluster;
diff --git a/Engine/Plugins/VirtualProduction/TextureShare/Source/TextureShareDisplayCluster/Private/PostProcess/TextureSharePostprocess.cpp b/Engine/Plugins/VirtualProduction/TextureShare/Source/TextureShareDisplayCluster/Private/PostProcess/TextureSharePostprocess.cpp
--- a/Engine/Plugins/VirtualProduction/TextureShare/Source/TextureShareDisplayCluster/Private/PostProcess/TextureSharePostprocess.cpp
+++ b/Engine/Plugins/VirtualProduction/TextureShare/Source/TextureShareDisplayCluster/Private/PostProcess/TextureSharePostprocess.cpp
@@ -28,35 +28,29 @@ static FAutoConsoleVariableRef CVarTextureShareEnableDisplayCluster(
 	ECVF_RenderThreadSafe
 );
 
-namespace UE
+namespace UE::TextureShare::PostProcess
 {
-	namespace TextureShare
+	static ITextureShareAPI& TextureShareAPI()
 	{
-		namespace PostProcess
-		{
-			static ITextureShareAPI& TextureShareAPI()
-			{
-				static ITextureShareAPI& TextureShareAPISingleton = ITextureShare::Get().GetTextureShareAPI();
-				return TextureShareAPISingleton;
-			}
+		static ITextureShareAPI& TextureShareAPISingleton = ITextureShare::Get().GetTextureShareAPI();
+		return TextureShareAPISingleton;
+	}
 
-			static FViewport* GetDisplayViewport(IDisplayClusterViewportManager* InViewportManager)
+	static FViewport* GetDisplayViewport(IDisplayClusterViewportManager* InViewportManager)
+	{
+		if (InViewportManager)
+		{
+			if (UWorld* CurrentWorld = InViewportManager->GetCurrentWorld())
 			{
-				if (InViewportManager)
+				if (UGameViewportClient* GameViewportClientPtr = CurrentWorld->GetGameViewport())
 				{
-					if (UWorld* CurrentWorld = InViewportManager->GetCurrentWorld())
-					{
-						if (UGameViewportClient* GameViewportClientPtr = CurrentWorld->GetGameViewport())
-						{
-							return GameViewportClientPtr->Viewport;
-						}
-					}
+					return GameViewportClientPtr->Viewport;
 				}
-
-				return nullptr;
-			};
+			}
 		}
-	}
+
+		return nullptr;
+	};
 };
 using namespace UE::TextureShare::PostProcess;
 using namespace UE::TextureShare;
diff --git a/Engine/Plugins/VirtualProduction/TextureShare/Source/TextureShareDisplayCluster/Private/PostProcess/TextureSharePostprocess_FrameProxy.cpp b/Engine/Plugins/VirtualProduction/TextureShare/Source/TextureShareDisplayCluster/Private/PostProcess/TextureSharePostprocess_FrameProxy.cpp
--- a/Engine/Plugins/VirtualProduction/TextureShare/Source/TextureShareDisplayCluster/Private/PostProcess/TextureSharePostprocess_FrameProxy.cpp
+++ b/Engine/Plugins/VirtualProduction/TextureShare/Source/TextureShareDisplayCluster/Private/PostProcess/TextureSharePostprocess_FrameProxy.cpp
@@ -13,27 +13,21 @@
 #include "Render/Viewport/IDisplayClusterViewportProxy.h"
 #include "Render/Projection/IDisplayClusterProjectionPolicy.h"
 
-namespace UE
+namespace UE::TextureShare::PostProcess_FrameProxy
 {
-	namespace TextureShare
+	// Support warp blend logic
+	static bool ShouldApplyWarpBlend(IDisplayClusterViewportProxy* ViewportProxy)
 	{
-		namespace PostProcess_FrameProxy
+		if (ViewportProxy->GetPostRenderSettings_RenderThread().Replace.IsEnabled())
 		{
-			// Support warp blend logic
-			static bool ShouldApplyWarpBlend(IDisplayClusterViewportProxy* ViewportProxy)
-			{
-				if (ViewportProxy->GetPostRenderSettings_RenderThread().Replace.IsEnabled())
-				{
-					// When used override texture, disable warp blend
-					return false;
-				}
+			// When used override texture, disable warp blend
+			return false;
+		}
 
-				const TSharedPtr<IDisplayClusterProjectionPolicy, ESPMode::ThreadSafe>& PrjPolicy = ViewportProxy->GetProjectionPolicy_RenderThread();
+		const TSharedPtr<IDisplayClusterProjectionPolicy, ESPMode::ThreadSafe>& PrjPolicy = ViewportProxy->GetProjectionPolicy_RenderThread();
 
-				// Projection policy must support warp blend op
-				return PrjPolicy.IsValid() && PrjPolicy->IsWarpBlendSupported();
-			}
-		}
+		// Projection policy must support warp blend op
+		return PrjPolicy.IsValid() && PrjPolicy->IsWarpBlendSupported();
 	}
 };
 using namespace UE::TextureShare::PostProcess_FrameProxy;
