Fixed shadowing warnings.

diff --git a/Engine/Plugins/VirtualProduction/TextureShare/Source/TextureShare/Private/Game/WorldSubsystem/TextureShareWorldSubsystemObjectProxy.cpp b/Engine/Plugins/VirtualProduction/TextureShare/Source/TextureShare/Private/Game/WorldSubsystem/TextureShareWorldSubsystemObjectProxy.cpp
--- a/Engine/Plugins/VirtualProduction/TextureShare/Source/TextureShare/Private/Game/WorldSubsystem/TextureShareWorldSubsystemObjectProxy.cpp
+++ b/Engine/Plugins/VirtualProduction/TextureShare/Source/TextureShare/Private/Game/WorldSubsystem/TextureShareWorldSubsystemObjectProxy.cpp
@@ -103,30 +103,30 @@ void FTextureShareWorldSubsystemObjectProxy::Update_RenderThread(FRHICommandList
 
 			TFunctionTextureShareViewExtension PreRenderViewFamilyFunction = [ProxyData](FRHICommandListImmediate& RHICmdList, FTextureShareSceneViewExtension& InViewExtension)
 			{
-				TSharedPtr<ITextureShareObjectProxy, ESPMode::ThreadSafe> ObjectProxy = InViewExtension.GetObjectProxy();
-				if (ObjectProxy.IsValid())
+				TSharedPtr<ITextureShareObjectProxy, ESPMode::ThreadSafe> ViewExtensionObjectProxy = InViewExtension.GetObjectProxy();
+				if (ViewExtensionObjectProxy.IsValid())
 				{
-					ObjectProxy->BeginFrameSync_RenderThread(RHICmdList);
+					ViewExtensionObjectProxy->BeginFrameSync_RenderThread(RHICmdList);
 				}
 			};
 
 			TFunctionTextureShareViewExtension PostRenderViewFamilyFunction = [ProxyData, bBackbufferShared](FRHICommandListImmediate& RHICmdList, FTextureShareSceneViewExtension& InViewExtension)
 			{
 				// Share user resources
-				TSharedPtr<ITextureShareObjectProxy, ESPMode::ThreadSafe> ObjectProxy = InViewExtension.GetObjectProxy();
-				if (ObjectProxy.IsValid()
-					&& ObjectProxy->IsFrameSyncActive_RenderThread())
+				TSharedPtr<ITextureShareObjectProxy, ESPMode::ThreadSafe> ViewExtensionObjectProxy = InViewExtension.GetObjectProxy();
+				if (ViewExtensionObjectProxy.IsValid()
+					&& ViewExtensionObjectProxy->IsFrameSyncActive_RenderThread())
 				{
-					FTextureShareWorldSubsystemObjectProxy ProxyAPI(ObjectProxy);
+					FTextureShareWorldSubsystemObjectProxy ProxyAPI(ViewExtensionObjectProxy);
 
 					// Share custom textures and sync
 					ProxyAPI.UpdateResources_RenderThread(RHICmdList, ProxyData);
-					ObjectProxy->FrameSync_RenderThread(RHICmdList, ETextureShareSyncStep::FrameProxyPreRenderEnd);
+					ViewExtensionObjectProxy->FrameSync_RenderThread(RHICmdList, ETextureShareSyncStep::FrameProxyPreRenderEnd);
 
 					// End frame when no back buffer is used
 					if (!bBackbufferShared)
 					{
-						ObjectProxy->EndFrameSync_RenderThread(RHICmdList);
+						ViewExtensionObjectProxy->EndFrameSync_RenderThread(RHICmdList);
 					}
 				}
 			};
@@ -136,18 +136,18 @@ void FTextureShareWorldSubsystemObjectProxy::Update_RenderThread(FRHICommandList
 				// Share backbuffer
 				if (bBackbufferShared)
 				{
-					TSharedPtr<ITextureShareObjectProxy, ESPMode::ThreadSafe> ObjectProxy = InViewExtension.GetObjectProxy();
-					if (ObjectProxy.IsValid()
-						&& ObjectProxy->IsFrameSyncActive_RenderThread())
+					TSharedPtr<ITextureShareObjectProxy, ESPMode::ThreadSafe> ViewExtensionObjectProxy = InViewExtension.GetObjectProxy();
+					if (ViewExtensionObjectProxy.IsValid()
+						&& ViewExtensionObjectProxy->IsFrameSyncActive_RenderThread())
 					{
-						FTextureShareWorldSubsystemObjectProxy ProxyAPI(ObjectProxy);
+						FTextureShareWorldSubsystemObjectProxy ProxyAPI(ViewExtensionObjectProxy);
 
 						// Share backbuffer and sync
 						ProxyAPI.UpdateFrameProxyBackbuffer_RenderThread(RHICmdList, InBackbuffer);
-						ObjectProxy->FrameSync_RenderThread(RHICmdList, ETextureShareSyncStep::FrameProxyBackBufferReadyToPresentEnd);
+						ViewExtensionObjectProxy->FrameSync_RenderThread(RHICmdList, ETextureShareSyncStep::FrameProxyBackBufferReadyToPresentEnd);
 
 						// End frame
-						ObjectProxy->EndFrameSync_RenderThread(RHICmdList);
+						ViewExtensionObjectProxy->EndFrameSync_RenderThread(RHICmdList);
 					}
 				}
 			};
