Changing manual include paths to use PublicIncludePathModuleNames.

diff --git a/Engine/Plugins/VirtualProduction/TextureShare/Source/TextureShareCore/TextureShareCore.Build.cs b/Engine/Plugins/VirtualProduction/TextureShare/Source/TextureShareCore/TextureShareCore.Build.cs
--- a/Engine/Plugins/VirtualProduction/TextureShare/Source/TextureShareCore/TextureShareCore.Build.cs
+++ b/Engine/Plugins/VirtualProduction/TextureShare/Source/TextureShareCore/TextureShareCore.Build.cs
@@ -1,7 +1,7 @@
 // Copyright Epic Games, Inc. All Rights Reserved.
 
+using EpicGames.Core;
 using UnrealBuildTool;
-using System.IO;
 
 [SupportedPlatforms("Win64")]
 public class TextureShareCore : ModuleRules
@@ -69,12 +69,7 @@ public class TextureShareCore : ModuleRules
 
 		if (Target.Platform.IsInGroup(UnrealPlatformGroup.Windows))
 		{
-			PublicIncludePaths.AddRange(
-				new string[] {
-				Path.Combine(EngineDirectory,"Source/Runtime/Windows/D3D11RHI/Public"),
-				Path.Combine(EngineDirectory,"Source/Runtime/D3D12RHI/Public"),
-				}
-			);
+			PublicIncludePathModuleNames.AddAll("D3D11RHI", "D3D12RHI");
 
 			AddEngineThirdPartyPrivateStaticDependencies(Target, "DX11");
 			AddEngineThirdPartyPrivateStaticDependencies(Target, "DX12");
