Remove #pragma once from cpp files
#rnx

diff --git a/Engine/Plugins/Experimental/TextureGraph/Source/TextureGraph/Private/Blueprint/TG_AsyncExportTask.cpp b/Engine/Plugins/Experimental/TextureGraph/Source/TextureGraph/Private/Blueprint/TG_AsyncExportTask.cpp
--- a/Engine/Plugins/Experimental/TextureGraph/Source/TextureGraph/Private/Blueprint/TG_AsyncExportTask.cpp
+++ b/Engine/Plugins/Experimental/TextureGraph/Source/TextureGraph/Private/Blueprint/TG_AsyncExportTask.cpp
@@ -1,6 +1,5 @@
 // Copyright Epic Games, Inc. All Rights Reserved.
 
-#pragma once
 #include "Blueprint/TG_AsyncExportTask.h"
 #include "Misc/Paths.h"
 #include "Engine/Engine.h"
diff --git a/Engine/Plugins/Experimental/TextureGraph/Source/TextureGraph/Private/Blueprint/TG_AsyncRenderTask.cpp b/Engine/Plugins/Experimental/TextureGraph/Source/TextureGraph/Private/Blueprint/TG_AsyncRenderTask.cpp
--- a/Engine/Plugins/Experimental/TextureGraph/Source/TextureGraph/Private/Blueprint/TG_AsyncRenderTask.cpp
+++ b/Engine/Plugins/Experimental/TextureGraph/Source/TextureGraph/Private/Blueprint/TG_AsyncRenderTask.cpp
@@ -1,6 +1,5 @@
 // Copyright Epic Games, Inc. All Rights Reserved.
 
-#pragma once
 #include "Blueprint/TG_AsyncRenderTask.h"
 #include "Engine/Engine.h"
 #include "Engine/World.h"
diff --git a/Engine/Plugins/Experimental/TextureGraph/Source/TextureGraph/Private/Blueprint/TG_AsyncTask.cpp b/Engine/Plugins/Experimental/TextureGraph/Source/TextureGraph/Private/Blueprint/TG_AsyncTask.cpp
--- a/Engine/Plugins/Experimental/TextureGraph/Source/TextureGraph/Private/Blueprint/TG_AsyncTask.cpp
+++ b/Engine/Plugins/Experimental/TextureGraph/Source/TextureGraph/Private/Blueprint/TG_AsyncTask.cpp
@@ -1,6 +1,5 @@
 // Copyright Epic Games, Inc. All Rights Reserved.
 
-#pragma once
 #include "Blueprint/TG_AsyncTask.h"
 
 #include "MixerEngine.h"
diff --git a/Engine/Plugins/Experimental/TextureGraph/Source/TextureGraph/Private/TG_HelperFunctions.cpp b/Engine/Plugins/Experimental/TextureGraph/Source/TextureGraph/Private/TG_HelperFunctions.cpp
--- a/Engine/Plugins/Experimental/TextureGraph/Source/TextureGraph/Private/TG_HelperFunctions.cpp
+++ b/Engine/Plugins/Experimental/TextureGraph/Source/TextureGraph/Private/TG_HelperFunctions.cpp
@@ -1,6 +1,5 @@
 // Copyright Epic Games, Inc. All Rights Reserved.
 
-#pragma once
 #include "TG_HelperFunctions.h"
 #include "Job/Scheduler.h"
 #include "TG_Texture.h"
diff --git a/Engine/Plugins/Runtime/nDisplay/Source/DisplayClusterConfiguration/Private/DisplayClusterConfigurationTypes_Media.cpp b/Engine/Plugins/Runtime/nDisplay/Source/DisplayClusterConfiguration/Private/DisplayClusterConfigurationTypes_Media.cpp
--- a/Engine/Plugins/Runtime/nDisplay/Source/DisplayClusterConfiguration/Private/DisplayClusterConfigurationTypes_Media.cpp
+++ b/Engine/Plugins/Runtime/nDisplay/Source/DisplayClusterConfiguration/Private/DisplayClusterConfigurationTypes_Media.cpp
@@ -1,7 +1,5 @@
 // Copyright Epic Games, Inc. All Rights Reserved.
 
-#pragma once
-
 #include "DisplayClusterConfigurationTypes_Media.h"
 
 ///////////////////////////////////////////////////
diff --git a/Engine/Plugins/VirtualProduction/TextureShare/Source/TextureShare/Private/Blueprints/TextureShareBlueprintContainers.cpp b/Engine/Plugins/VirtualProduction/TextureShare/Source/TextureShare/Private/Blueprints/TextureShareBlueprintContainers.cpp
--- a/Engine/Plugins/VirtualProduction/TextureShare/Source/TextureShare/Private/Blueprints/TextureShareBlueprintContainers.cpp
+++ b/Engine/Plugins/VirtualProduction/TextureShare/Source/TextureShare/Private/Blueprints/TextureShareBlueprintContainers.cpp
@@ -1,6 +1,5 @@
 // Copyright Epic Games, Inc. All Rights Reserved.
 
-#pragma once
 #include "Blueprints/TextureShareBlueprintContainers.h"
 #include "Misc/TextureShareCoreStrings.h"
 #include "UObject/Package.h"
diff --git a/Engine/Plugins/VirtualProduction/TextureShare/Source/TextureShareCore/Private/IPC/Containers/TextureShareCoreInterprocessMemory.cpp b/Engine/Plugins/VirtualProduction/TextureShare/Source/TextureShareCore/Private/IPC/Containers/TextureShareCoreInterprocessMemory.cpp
--- a/Engine/Plugins/VirtualProduction/TextureShare/Source/TextureShareCore/Private/IPC/Containers/TextureShareCoreInterprocessMemory.cpp
+++ b/Engine/Plugins/VirtualProduction/TextureShare/Source/TextureShareCore/Private/IPC/Containers/TextureShareCoreInterprocessMemory.cpp
@@ -1,6 +1,5 @@
 // Copyright Epic Games, Inc. All Rights Reserved.
 
-#pragma once
 #include "IPC/Containers/TextureShareCoreInterprocessMemory.h"
 #include "IPC/TextureShareCoreInterprocessHelpers.h"
 #include "Module/TextureShareCoreLog.h"
diff --git a/Engine/Plugins/VirtualProduction/TextureShare/Source/TextureShareCore/Private/IPC/Containers/TextureShareCoreInterprocessObjectData.cpp b/Engine/Plugins/VirtualProduction/TextureShare/Source/TextureShareCore/Private/IPC/Containers/TextureShareCoreInterprocessObjectData.cpp
--- a/Engine/Plugins/VirtualProduction/TextureShare/Source/TextureShareCore/Private/IPC/Containers/TextureShareCoreInterprocessObjectData.cpp
+++ b/Engine/Plugins/VirtualProduction/TextureShare/Source/TextureShareCore/Private/IPC/Containers/TextureShareCoreInterprocessObjectData.cpp
@@ -1,6 +1,5 @@
 // Copyright Epic Games, Inc. All Rights Reserved.
 
-#pragma once
 #include "IPC/Containers/TextureShareCoreInterprocessObjectData.h"
 #include "IPC/TextureShareCoreInterprocessSerializer.h"
 
diff --git a/Engine/Plugins/VirtualProduction/TextureShare/Source/TextureShareCore/Private/IPC/Containers/TextureShareCoreInterprocessObjectDesc.cpp b/Engine/Plugins/VirtualProduction/TextureShare/Source/TextureShareCore/Private/IPC/Containers/TextureShareCoreInterprocessObjectDesc.cpp
--- a/Engine/Plugins/VirtualProduction/TextureShare/Source/TextureShareCore/Private/IPC/Containers/TextureShareCoreInterprocessObjectDesc.cpp
+++ b/Engine/Plugins/VirtualProduction/TextureShare/Source/TextureShareCore/Private/IPC/Containers/TextureShareCoreInterprocessObjectDesc.cpp
@@ -1,6 +1,5 @@
 // Copyright Epic Games, Inc. All Rights Reserved.
 
-#pragma once
 #include "IPC/Containers/TextureShareCoreInterprocessObjectDesc.h"
 
 /////////////////////////////////////////////////////////////////////////////////////
diff --git a/Engine/Plugins/VirtualProduction/TextureShare/Source/TextureShareCore/Private/IPC/Containers/TextureShareCoreInterprocessObjectSync.cpp b/Engine/Plugins/VirtualProduction/TextureShare/Source/TextureShareCore/Private/IPC/Containers/TextureShareCoreInterprocessObjectSync.cpp
--- a/Engine/Plugins/VirtualProduction/TextureShare/Source/TextureShareCore/Private/IPC/Containers/TextureShareCoreInterprocessObjectSync.cpp
+++ b/Engine/Plugins/VirtualProduction/TextureShare/Source/TextureShareCore/Private/IPC/Containers/TextureShareCoreInterprocessObjectSync.cpp
@@ -1,6 +1,5 @@
 // Copyright Epic Games, Inc. All Rights Reserved.
 
-#pragma once
 #include "IPC/Containers/TextureShareCoreInterprocessObjectSync.h"
 
 /////////////////////////////////////////////////////////////////////////////////////
