TextureShare: Fixed a bug with incorrect scene texture sizes.

diff --git a/Engine/Plugins/VirtualProduction/TextureShare/Source/TextureShare/Private/Game/ViewExtension/TextureShareSceneViewExtension.cpp b/Engine/Plugins/VirtualProduction/TextureShare/Source/TextureShare/Private/Game/ViewExtension/TextureShareSceneViewExtension.cpp
--- a/Engine/Plugins/VirtualProduction/TextureShare/Source/TextureShare/Private/Game/ViewExtension/TextureShareSceneViewExtension.cpp
+++ b/Engine/Plugins/VirtualProduction/TextureShare/Source/TextureShare/Private/Game/ViewExtension/TextureShareSceneViewExtension.cpp
@@ -133,17 +133,20 @@ void FTextureShareSceneViewExtension::GetSceneViewData_RenderThread(const FTextu
 
 void FTextureShareSceneViewExtension::ShareSceneViewColors_RenderThread(FRDGBuilder& GraphBuilder, const FSceneTextures& SceneTextures, const FTextureShareSceneView& InView)
 {
-	const auto AddShareTexturePass = [&](const TCHAR* InTextureName, const FRDGTextureRef& InTextureRef)
+	const auto AddShareTexturePass = [&](const TCHAR* InTextureName, const FRDGTextureRef& InTextureRef, const FIntRect* CustomSrcRect = nullptr)
 	{
 		// Send resource
 		ObjectProxy->ShareResource_RenderThread(GraphBuilder, FTextureShareCoreResourceDesc(InTextureName, InView.ViewInfo.ViewDesc, ETextureShareTextureOp::Read),
-			InTextureRef, InView.GPUIndex, &InView.UnconstrainedViewRect);
+			InTextureRef, InView.GPUIndex, CustomSrcRect  ? CustomSrcRect : &InView.UnconstrainedViewRect);
 	};
 
 	AddShareTexturePass(UE::TextureShareStrings::SceneTextures::SceneColor, SceneTextures.Color.Resolve);
 
 	AddShareTexturePass(UE::TextureShareStrings::SceneTextures::SceneDepth, SceneTextures.Depth.Resolve);
-	AddShareTexturePass(UE::TextureShareStrings::SceneTextures::SmallDepthZ, SceneTextures.SmallDepth);
+	
+	// Small depth use downscale size
+	const FIntRect SmallDepthRect = GetDownscaledRect(InView.UnconstrainedViewRect, SceneTextures.Config.SmallDepthDownsampleFactor);
+	AddShareTexturePass(UE::TextureShareStrings::SceneTextures::SmallDepthZ, SceneTextures.SmallDepth, &SmallDepthRect);
 
 	AddShareTexturePass(UE::TextureShareStrings::SceneTextures::GBufferA, SceneTextures.GBufferA);
 	AddShareTexturePass(UE::TextureShareStrings::SceneTextures::GBufferB, SceneTextures.GBufferB);
diff --git a/Engine/Plugins/VirtualProduction/TextureShare/Source/TextureShare/Private/Resources/TextureShareResourcesProxy.cpp b/Engine/Plugins/VirtualProduction/TextureShare/Source/TextureShare/Private/Resources/TextureShareResourcesProxy.cpp
--- a/Engine/Plugins/VirtualProduction/TextureShare/Source/TextureShare/Private/Resources/TextureShareResourcesProxy.cpp
+++ b/Engine/Plugins/VirtualProduction/TextureShare/Source/TextureShare/Private/Resources/TextureShareResourcesProxy.cpp
@@ -33,6 +33,22 @@
 
 namespace UE::TextureShare::ResourcesProxy
 {
+	static inline FIntRect GetTextureRect(FRHITexture* InTexture, const FIntRect* InRect)
+	{
+		// InRect must be within the size of InTexture.
+		FIntRect OutRect(FIntPoint(EForceInit::ForceInitToZero), InTexture->GetSizeXY());
+		if (InRect)
+		{
+			OutRect.Max.X = FMath::Min(OutRect.Max.X, InRect->Max.X);
+			OutRect.Max.Y = FMath::Min(OutRect.Max.Y, InRect->Max.Y);
+
+			OutRect.Min.X = FMath::Min(OutRect.Max.X, InRect->Min.X);
+			OutRect.Min.Y = FMath::Min(OutRect.Max.Y, InRect->Min.Y);
+		}
+
+		return OutRect;
+	}
+
 	static bool IsSizeResampleRequired(FRHITexture* SrcTexture, FRHITexture* DstTexture, const FIntRect* SrcTextureRect, const FIntRect* DstTextureRect, FIntRect& OutSrcRect, FIntRect& OutDstRect)
 	{
 		FIntVector SrcSizeXYZ = SrcTexture->GetSizeXYZ();
@@ -365,7 +381,7 @@ bool FTextureShareResourcesProxy::GetPooledTempRTT_RenderThread(FRHICommandListI
 	return false;
 }
 
-bool FTextureShareResourcesProxy::WriteToShareTexture_RenderThread(FRHICommandListImmediate& RHICmdList, FRHITexture* InSrcTexture, const FIntRect* InSrcTextureRect, FTextureShareResource* InDestSharedResource)
+bool FTextureShareResourcesProxy::WriteToShareTexture_RenderThread(FRHICommandListImmediate& RHICmdList, FRHITexture* InSrcTexture, const FIntRect* InSrcTextureRectPtr, FTextureShareResource* InDestSharedResource)
 {
 	if (InSrcTexture && InDestSharedResource)
 	{
@@ -373,16 +389,17 @@ bool FTextureShareResourcesProxy::WriteToShareTexture_RenderThread(FRHICommandLi
 		{
 			UE_TS_LOG(LogTextureShareResource, Log, TEXT("%s:WriteToShareTexture_RenderThread(%s.%s)"), *InDestSharedResource->GetCoreObjectName(), *InDestSharedResource->GetResourceDesc().ViewDesc.Id, *InDestSharedResource->GetResourceDesc().ResourceName);
 
+			const FIntRect SrcTextureRect = GetTextureRect(InSrcTexture, InSrcTextureRectPtr);
 			const EPixelFormat InDestFormat = InDestSharedTexture->GetFormat();
 			const bool bIsFormatResampleRequired = InSrcTexture->GetFormat() != InDestFormat;
 
 			FIntRect SrcRect, DestRect;
-			const bool bResampleRequired = IsSizeResampleRequired(InSrcTexture, InDestSharedTexture, InSrcTextureRect, nullptr, SrcRect, DestRect) || bIsFormatResampleRequired;
+			const bool bResampleRequired = IsSizeResampleRequired(InSrcTexture, InDestSharedTexture, &SrcTextureRect, nullptr, SrcRect, DestRect) || bIsFormatResampleRequired;
 			if (!bResampleRequired)
 			{
 				// Copy direct to shared texture
 				SCOPE_CYCLE_COUNTER(STAT_TextureShare_CopyShared);
-				DirectCopyTexture_RenderThread(RHICmdList, InSrcTexture, InDestSharedTexture, InSrcTextureRect, nullptr);
+				DirectCopyTexture_RenderThread(RHICmdList, InSrcTexture, InDestSharedTexture, &SrcTextureRect, nullptr);
 
 				bRHIThreadChanged = true;
 
@@ -398,7 +415,7 @@ bool FTextureShareResourcesProxy::WriteToShareTexture_RenderThread(FRHICommandLi
 					{
 						// Resample source texture to PooledTempRTT (Src texture now always shader resource)
 						SCOPE_CYCLE_COUNTER(STAT_TextureShare_ResampleTempRTT);
-						ResampleCopyTexture_RenderThread(RHICmdList, InSrcTexture, RHIResampledRTT, InSrcTextureRect, nullptr);
+						ResampleCopyTexture_RenderThread(RHICmdList, InSrcTexture, RHIResampledRTT, &SrcTextureRect, nullptr);
 
 						// Copy PooledTempRTT to shared texture surface
 						SCOPE_CYCLE_COUNTER(STAT_TextureShare_CopyShared);
@@ -416,7 +433,7 @@ bool FTextureShareResourcesProxy::WriteToShareTexture_RenderThread(FRHICommandLi
 	return false;
 }
 
-bool FTextureShareResourcesProxy::ReadFromShareTexture_RenderThread(FRHICommandListImmediate& RHICmdList, FTextureShareResource* InSrcSharedResource, FRHITexture* InDestTexture, const FIntRect* InDestTextureRect)
+bool FTextureShareResourcesProxy::ReadFromShareTexture_RenderThread(FRHICommandListImmediate& RHICmdList, FTextureShareResource* InSrcSharedResource, FRHITexture* InDestTexture, const FIntRect* InDestTextureRectPtr)
 {
 	if (InSrcSharedResource && InDestTexture)
 	{
@@ -424,17 +441,18 @@ bool FTextureShareResourcesProxy::ReadFromShareTexture_RenderThread(FRHICommandL
 		{
 			UE_TS_LOG(LogTextureShareResource, Log, TEXT("%s:ReadFromShareTexture_RenderThread(%s.%s)"), *InSrcSharedResource->GetCoreObjectName(), *InSrcSharedResource->GetResourceDesc().ViewDesc.Id, *InSrcSharedResource->GetResourceDesc().ResourceName);
 
+			const FIntRect DestTextureRect = GetTextureRect(InDestTexture, InDestTextureRectPtr);
 			const EPixelFormat InSrcFormat = InSrcSharedTexture->GetFormat();
 			const EPixelFormat InDestFormat = InDestTexture->GetFormat();
 			const bool bIsFormatResampleRequired = InSrcFormat != InDestFormat;
 
 			FIntRect SrcRect, DestRect;
-			const bool bResampleRequired = IsSizeResampleRequired(InSrcSharedTexture, InDestTexture, nullptr, InDestTextureRect, SrcRect, DestRect) || bIsFormatResampleRequired;
+			const bool bResampleRequired = IsSizeResampleRequired(InSrcSharedTexture, InDestTexture, nullptr, &DestTextureRect, SrcRect, DestRect) || bIsFormatResampleRequired;
 			if (!bResampleRequired)
 			{
 				// Copy direct from shared texture
 				SCOPE_CYCLE_COUNTER(STAT_TextureShare_CopyShared);
-				DirectCopyTexture_RenderThread(RHICmdList, InSrcSharedTexture, InDestTexture, nullptr, InDestTextureRect);
+				DirectCopyTexture_RenderThread(RHICmdList, InSrcSharedTexture, InDestTexture, nullptr, &DestTextureRect);
 
 				bRHIThreadChanged = true;
 
@@ -460,7 +478,7 @@ bool FTextureShareResourcesProxy::ReadFromShareTexture_RenderThread(FRHICommandL
 							ResampleCopyTexture_RenderThread(RHICmdList, RHIReceivedSRV, RHIResampledRTT, nullptr, nullptr);
 
 							// Copy RHIResampledRTT to Destination
-							DirectCopyTexture_RenderThread(RHICmdList, RHIResampledRTT, InDestTexture, nullptr, InDestTextureRect);
+							DirectCopyTexture_RenderThread(RHICmdList, RHIResampledRTT, InDestTexture, nullptr, &DestTextureRect);
 
 							bRHIThreadChanged = true;
 
