Use GetNonValidationRHI for casting to the base RHI class.

diff --git a/Engine/Plugins/Media/ElectraUtil/Source/ElectraSamples/Private/ElectraTextureSamplePC.cpp b/Engine/Plugins/Media/ElectraUtil/Source/ElectraSamples/Private/ElectraTextureSamplePC.cpp
--- a/Engine/Plugins/Media/ElectraUtil/Source/ElectraSamples/Private/ElectraTextureSamplePC.cpp
+++ b/Engine/Plugins/Media/ElectraUtil/Source/ElectraSamples/Private/ElectraTextureSamplePC.cpp
@@ -228,7 +228,7 @@ FRHITexture* FElectraTextureSample::GetTexture() const
 					{
 						Flags = Flags | ETextureCreateFlags::SRGB;
 					}
-					Texture = static_cast<ID3D12DynamicRHI*>(GDynamicRHI)->RHICreateTexture2DFromResource(VideoDecoderOutput->GetFormat(), Flags, FClearValueBinding(), TextureDX12);
+					Texture = static_cast<ID3D12DynamicRHI*>(GDynamicRHI->GetNonValidationRHI())->RHICreateTexture2DFromResource(VideoDecoderOutput->GetFormat(), Flags, FClearValueBinding(), TextureDX12);
 				}
 			}
 		}
diff --git a/Engine/Plugins/VirtualProduction/TextureShare/Source/TextureShare/Private/Resources/TextureShareResource_Vulkan.cpp b/Engine/Plugins/VirtualProduction/TextureShare/Source/TextureShare/Private/Resources/TextureShareResource_Vulkan.cpp
--- a/Engine/Plugins/VirtualProduction/TextureShare/Source/TextureShare/Private/Resources/TextureShareResource_Vulkan.cpp
+++ b/Engine/Plugins/VirtualProduction/TextureShare/Source/TextureShare/Private/Resources/TextureShareResource_Vulkan.cpp
@@ -26,7 +26,7 @@ namespace UE::TextureShare::Resource_Vulkan
 
 	static FTextureShareDeviceVulkanContext GetDeviceVulkanContext()
 	{
-		if (FVulkanDynamicRHI* DynamicRHI = static_cast<FVulkanDynamicRHI*>(GDynamicRHI))
+		if (FVulkanDynamicRHI* DynamicRHI = static_cast<FVulkanDynamicRHI*>(GDynamicRHI->GetNonValidationRHI()))
 		{
 			return FTextureShareDeviceVulkanContext(DynamicRHI->GetInstance(), DynamicRHI->GetDevice()->GetPhysicalHandle(), DynamicRHI->GetDevice()->GetInstanceHandle());
 		}
