[TextureShare] [nDisplay] Fixed numerous of the bugs:

Multithread data issues:
- The API has been slightly changed: the returned data no longer references MT data.
- Data that is accessed in multiple threads have been wrapped with functions with locks.

TextureShare SDK
- TextureShare SDK compilation: Added missing module for recent code changes (TextureShareSDK.Build.cs, module 'ImageCore')

UE issues:
- Removed LOG_SCOPE_VERBOSITY_OVERRIDE() from multithreaded logic. This function only works on a single thread. (TextureShareCoreInterprocessMutex.cpp)

diff --git a/Engine/Extras/VirtualProduction/TextureShare/Samples/Helpers/Private/TextureShareObject_Sync.cpp b/Engine/Extras/VirtualProduction/TextureShare/Samples/Helpers/Private/TextureShareObject_Sync.cpp
--- a/Engine/Extras/VirtualProduction/TextureShare/Samples/Helpers/Private/TextureShareObject_Sync.cpp
+++ b/Engine/Extras/VirtualProduction/TextureShare/Samples/Helpers/Private/TextureShareObject_Sync.cpp
@@ -54,18 +54,24 @@ bool FTextureShareObject::BeginFrame()
 
 bool FTextureShareObject::EndFrame()
 {
+	bool bEndFrameResult = true;
+
 	if (TextureShareSDKObject->IsFrameSyncActive_RenderThread())
 	{
 		// Flush deffered sync steps
 		if (!DefferedResourceSync_RenderThread(ETextureShareSyncStep::FrameProxyFlush))
 		{
-			return false;
+			bEndFrameResult = false;
 		}
+	}
 
-		return TextureShareSDKObject->EndFrameSync_RenderThread();
+	// Always call EndFrameSync_RenderThread().
+	if (!TextureShareSDKObject->EndFrameSync_RenderThread())
+	{
+		bEndFrameResult = false;
 	}
 
-	return false;
+	return bEndFrameResult;
 }
 
 bool FTextureShareObject::FrameSync(const ETextureShareSyncStep InSyncStep)
diff --git a/Engine/Extras/VirtualProduction/TextureShare/TextureShareSDK/Source/Containers/TextureShareCoreContainers_ManualProjection.h b/Engine/Extras/VirtualProduction/TextureShare/TextureShareSDK/Source/Containers/TextureShareCoreContainers_ManualProjection.h
--- a/Engine/Extras/VirtualProduction/TextureShare/TextureShareSDK/Source/Containers/TextureShareCoreContainers_ManualProjection.h
+++ b/Engine/Extras/VirtualProduction/TextureShare/TextureShareSDK/Source/Containers/TextureShareCoreContainers_ManualProjection.h
@@ -1,6 +1,8 @@
 // Copyright Epic Games, Inc. All Rights Reserved.
 
 #pragma once
+
+#include "Containers/TextureShareCoreContainers_FrameMarker.h"
 #include "Containers/TextureShareCoreContainers_ViewDesc.h"
 
 /**
diff --git a/Engine/Extras/VirtualProduction/TextureShare/TextureShareSDK/Source/Containers/TextureShareCoreEnums.h b/Engine/Extras/VirtualProduction/TextureShare/TextureShareSDK/Source/Containers/TextureShareCoreEnums.h
--- a/Engine/Extras/VirtualProduction/TextureShare/TextureShareSDK/Source/Containers/TextureShareCoreEnums.h
+++ b/Engine/Extras/VirtualProduction/TextureShare/TextureShareSDK/Source/Containers/TextureShareCoreEnums.h
@@ -58,8 +58,13 @@ enum class ETextureShareThreadMutex : uint8
 	RenderingThread,
 
 	InternalLock,
+	InternalObjectDescMT,
+	InternalFrameConnectionsMT,
+	InternalSyncSettingsMT,
 
-	COUNT
+	COUNT,
+	// All mutexes after this dont produce log
+	MAX_LOG = InternalLock
 };
 
 /**
diff --git a/Engine/Plugins/VirtualProduction/TextureShare/Source/TextureShare/Private/Game/WorldSubsystem/TextureShareWorldSubsystemContext.cpp b/Engine/Plugins/VirtualProduction/TextureShare/Source/TextureShare/Private/Game/WorldSubsystem/TextureShareWorldSubsystemContext.cpp
--- a/Engine/Plugins/VirtualProduction/TextureShare/Source/TextureShare/Private/Game/WorldSubsystem/TextureShareWorldSubsystemContext.cpp
+++ b/Engine/Plugins/VirtualProduction/TextureShare/Source/TextureShare/Private/Game/WorldSubsystem/TextureShareWorldSubsystemContext.cpp
@@ -97,8 +97,7 @@ namespace UE::TextureShareWorldSubsystemContext
 void FTextureShareWorldSubsystemContext::Tick_GameThread(ITextureShareObject& TextureShareObject)
 {
 	using namespace UE::TextureShareWorldSubsystemContext;
-	if (TextureShareObject.BeginFrameSync()
-		&& TextureShareObject.IsFrameSyncActive())
+	if (TextureShareObject.BeginFrameSync())
 	{
 		// Update frame marker for current frame
 		TextureShareObject.GetCoreData().FrameMarker.NextFrame();
@@ -114,8 +113,7 @@ void FTextureShareWorldSubsystemContext::Tick_GameThread(ITextureShareObject& Te
 			}
 		}
 
-		if (TextureShareObject.FrameSync(ETextureShareSyncStep::FramePreSetupBegin)
-			&& TextureShareObject.IsFrameSyncActive())
+		if (TextureShareObject.FrameSync(ETextureShareSyncStep::FramePreSetupBegin))
 		{
 			// Process resources requests.
 			{
@@ -145,13 +143,13 @@ void FTextureShareWorldSubsystemContext::Tick_GameThread(ITextureShareObject& Te
 					}
 				}
 			}
+		}
 
-			bGameThreadSynchronized = TextureShareObject.EndFrameSync();
+		bGameThreadSynchronized = TextureShareObject.EndFrameSync();
 
-			UE_TS_LOG(LogTextureShareWorldSubsystem, Log, TEXT("%s:Tick_GameThread() %s"),
-				*TextureShareObject.GetName(),
-				bGameThreadSynchronized ? TEXT("successfully completed") : TEXT("Failed."));
-		}
+		UE_TS_LOG(LogTextureShareWorldSubsystem, Log, TEXT("%s:Tick_GameThread() %s"),
+			*TextureShareObject.GetName(),
+			bGameThreadSynchronized ? TEXT("successfully completed") : TEXT("Failed."));
 	}
 }
 
diff --git a/Engine/Plugins/VirtualProduction/TextureShare/Source/TextureShare/Private/ITextureShareObject.h b/Engine/Plugins/VirtualProduction/TextureShare/Source/TextureShare/Private/ITextureShareObject.h
--- a/Engine/Plugins/VirtualProduction/TextureShare/Source/TextureShare/Private/ITextureShareObject.h
+++ b/Engine/Plugins/VirtualProduction/TextureShare/Source/TextureShare/Private/ITextureShareObject.h
@@ -41,7 +41,7 @@ public:
 	/**
 	* Returns detailed information about the TextureShare object.
 	*/
-	virtual const FTextureShareCoreObjectDesc& GetObjectDesc() const = 0;
+	virtual FTextureShareCoreObjectDesc GetObjectDesc() const = 0;
 
 	/**
 	 * Returns true if the TextureShare object is ready to be used.
@@ -79,7 +79,7 @@ public:
 	/**
 	* Return the sync settings of the TextureShare.
 	*/
-	virtual const FTextureShareCoreSyncSettings& GetSyncSetting() const = 0;
+	virtual FTextureShareCoreSyncSettings GetSyncSetting() const = 0;
 
 	/**
 	 * Get TextureShare default sync settings by template type
@@ -153,7 +153,7 @@ public:
 	 *
 	 * @return true, if success
 	 */
-	virtual const TArray<FTextureShareCoreObjectDesc>& GetConnectedInterprocessObjects() const = 0;
+	virtual TArray<FTextureShareCoreObjectDesc> GetConnectedInterprocessObjects() const = 0;
 
 public:
 	/////////////////////////// Data Containers ///////////////////////////
@@ -168,7 +168,7 @@ public:
 	/**
 	 * Received Data from connected process objects
 	 */
-	virtual const TArray<FTextureShareCoreObjectData>& GetReceivedCoreObjectData() const = 0;
+	virtual TArray<FTextureShareCoreObjectData> GetReceivedCoreObjectData() const = 0;
 
 	/**
 	 * Reference to views data
diff --git a/Engine/Plugins/VirtualProduction/TextureShare/Source/TextureShare/Private/ITextureShareObjectProxy.h b/Engine/Plugins/VirtualProduction/TextureShare/Source/TextureShare/Private/ITextureShareObjectProxy.h
--- a/Engine/Plugins/VirtualProduction/TextureShare/Source/TextureShare/Private/ITextureShareObjectProxy.h
+++ b/Engine/Plugins/VirtualProduction/TextureShare/Source/TextureShare/Private/ITextureShareObjectProxy.h
@@ -35,7 +35,7 @@ public:
 	/**
 	* Returns detailed information about the TextureShare object.
 	*/
-	virtual const FTextureShareCoreObjectDesc& GetObjectDesc_RenderThread() const = 0;
+	virtual FTextureShareCoreObjectDesc GetObjectDesc_RenderThread() const = 0;
 
 	/**
 	 * Returns true if the TextureShare object is ready to be used.
@@ -190,7 +190,7 @@ public:
 	/**
 	 * Received ProxyData from connected process objects
 	 */
-	virtual const TArray<FTextureShareCoreObjectProxyData>& GetReceivedCoreObjectProxyData_RenderThread() const = 0;
+	virtual TArray<FTextureShareCoreObjectProxyData> GetReceivedCoreObjectProxyData_RenderThread() const = 0;
 
 	/**
 	*  Return proxy flags
diff --git a/Engine/Plugins/VirtualProduction/TextureShare/Source/TextureShare/Private/Object/TextureShareObject.cpp b/Engine/Plugins/VirtualProduction/TextureShare/Source/TextureShare/Private/Object/TextureShareObject.cpp
--- a/Engine/Plugins/VirtualProduction/TextureShare/Source/TextureShare/Private/Object/TextureShareObject.cpp
+++ b/Engine/Plugins/VirtualProduction/TextureShare/Source/TextureShare/Private/Object/TextureShareObject.cpp
@@ -36,7 +36,7 @@ void FTextureShareObject::SetTextureShareContext(const TSharedPtr<ITextureShareC
 	}
 }
 
-ITextureShareContext* FTextureShareObject::GetTextureShareContext() const
+const ITextureShareContext* FTextureShareObject::GetTextureShareContext() const
 {
 	return TextureShareContext.Get();
 }
@@ -44,10 +44,10 @@ ITextureShareContext* FTextureShareObject::GetTextureShareContext() const
 //////////////////////////////////////////////////////////////////////////////////////////////
 const FString& FTextureShareObject::GetName() const
 {
-	return CoreObject->GetObjectDesc().ShareName;
+	return CoreObject->GetName();
 }
 
-const FTextureShareCoreObjectDesc& FTextureShareObject::GetObjectDesc() const
+FTextureShareCoreObjectDesc FTextureShareObject::GetObjectDesc() const
 {
 	return CoreObject->GetObjectDesc();
 }
@@ -69,12 +69,12 @@ bool FTextureShareObject::SetProcessId(const FString& InProcessId)
 
 bool FTextureShareObject::SetSyncSetting(const FTextureShareCoreSyncSettings& InSyncSetting)
 {
-	return CoreObject->SetSyncSetting(InSyncSetting);
+	return CoreObject->SetSyncSettings(InSyncSetting);
 }
 
-const FTextureShareCoreSyncSettings& FTextureShareObject::GetSyncSetting() const
+FTextureShareCoreSyncSettings FTextureShareObject::GetSyncSetting() const
 {
-	return CoreObject->GetSyncSetting();
+	return CoreObject->GetSyncSettings();
 }
 
 FTextureShareCoreFrameSyncSettings FTextureShareObject::GetFrameSyncSettings(const ETextureShareFrameSyncTemplate InType) const
@@ -270,7 +270,7 @@ bool FTextureShareObject::DoFrameSync(const ETextureShareSyncStep InSyncStep)
 }
 
 //////////////////////////////////////////////////////////////////////////////////////////////
-const TArray<FTextureShareCoreObjectDesc>& FTextureShareObject::GetConnectedInterprocessObjects() const
+TArray<FTextureShareCoreObjectDesc> FTextureShareObject::GetConnectedInterprocessObjects() const
 {
 	return CoreObject->GetConnectedInterprocessObjects();
 }
@@ -285,7 +285,7 @@ const FTextureShareCoreData& FTextureShareObject::GetCoreData() const
 	return CoreObject->GetData();
 }
 
-const TArray<FTextureShareCoreObjectData>& FTextureShareObject::GetReceivedCoreObjectData() const
+TArray<FTextureShareCoreObjectData> FTextureShareObject::GetReceivedCoreObjectData() const
 {
 	return CoreObject->GetReceivedData();
 }
diff --git a/Engine/Plugins/VirtualProduction/TextureShare/Source/TextureShare/Private/Object/TextureShareObject.h b/Engine/Plugins/VirtualProduction/TextureShare/Source/TextureShare/Private/Object/TextureShareObject.h
--- a/Engine/Plugins/VirtualProduction/TextureShare/Source/TextureShare/Private/Object/TextureShareObject.h
+++ b/Engine/Plugins/VirtualProduction/TextureShare/Source/TextureShare/Private/Object/TextureShareObject.h
@@ -24,10 +24,10 @@ public:
 public:
 	// ~ITextureShareObject
 	virtual void SetTextureShareContext(const TSharedPtr<ITextureShareContext, ESPMode::ThreadSafe>& InTextureShareContext) override;
-	virtual ITextureShareContext* GetTextureShareContext() const  override;
+	virtual const ITextureShareContext* GetTextureShareContext() const  override;
 
 	virtual const FString& GetName() const override;
-	virtual const FTextureShareCoreObjectDesc& GetObjectDesc() const override;
+	virtual FTextureShareCoreObjectDesc GetObjectDesc() const override;
 
 	virtual bool IsActive() const override;
 	virtual bool IsFrameSyncActive() const override;
@@ -35,7 +35,7 @@ public:
 	virtual bool SetProcessId(const FString& InProcessId) override;
 	
 	virtual bool SetSyncSetting(const FTextureShareCoreSyncSettings& InSyncSetting) override;
-	virtual const FTextureShareCoreSyncSettings& GetSyncSetting() const override;
+	virtual FTextureShareCoreSyncSettings GetSyncSetting() const override;
 
 	virtual FTextureShareCoreFrameSyncSettings GetFrameSyncSettings(const ETextureShareFrameSyncTemplate InType) const override;
 
@@ -47,12 +47,12 @@ public:
 	virtual bool FrameSync(const ETextureShareSyncStep InSyncStep) override;
 	virtual bool EndFrameSync() override;
 
-	virtual const TArray<FTextureShareCoreObjectDesc>& GetConnectedInterprocessObjects() const override;
+	virtual TArray<FTextureShareCoreObjectDesc> GetConnectedInterprocessObjects() const override;
 
 	virtual FTextureShareCoreData& GetCoreData() override;
 	virtual const FTextureShareCoreData& GetCoreData() const override;
 
-	virtual const TArray<FTextureShareCoreObjectData>& GetReceivedCoreObjectData() const override;
+	virtual TArray<FTextureShareCoreObjectData> GetReceivedCoreObjectData() const override;
 	
 	virtual FTextureShareData& GetData() override;
 
diff --git a/Engine/Plugins/VirtualProduction/TextureShare/Source/TextureShare/Private/Object/TextureShareObjectProxy.cpp b/Engine/Plugins/VirtualProduction/TextureShare/Source/TextureShare/Private/Object/TextureShareObjectProxy.cpp
--- a/Engine/Plugins/VirtualProduction/TextureShare/Source/TextureShare/Private/Object/TextureShareObjectProxy.cpp
+++ b/Engine/Plugins/VirtualProduction/TextureShare/Source/TextureShare/Private/Object/TextureShareObjectProxy.cpp
@@ -161,7 +161,7 @@ const FString& FTextureShareObjectProxy::GetName_RenderThread() const
 	return CoreObject->GetName();
 }
 
-const FTextureShareCoreObjectDesc& FTextureShareObjectProxy::GetObjectDesc_RenderThread() const
+FTextureShareCoreObjectDesc FTextureShareObjectProxy::GetObjectDesc_RenderThread() const
 {
 	check(IsInRenderingThread());
 
@@ -411,7 +411,7 @@ const FTextureShareCoreProxyData& FTextureShareObjectProxy::GetCoreProxyData_Ren
 	return CoreObject->GetProxyData_RenderThread();
 }
 
-const TArray<FTextureShareCoreObjectProxyData>& FTextureShareObjectProxy::GetReceivedCoreObjectProxyData_RenderThread() const
+TArray<FTextureShareCoreObjectProxyData> FTextureShareObjectProxy::GetReceivedCoreObjectProxyData_RenderThread() const
 {
 	check(IsInRenderingThread());
 
diff --git a/Engine/Plugins/VirtualProduction/TextureShare/Source/TextureShare/Private/Object/TextureShareObjectProxy.h b/Engine/Plugins/VirtualProduction/TextureShare/Source/TextureShare/Private/Object/TextureShareObjectProxy.h
--- a/Engine/Plugins/VirtualProduction/TextureShare/Source/TextureShare/Private/Object/TextureShareObjectProxy.h
+++ b/Engine/Plugins/VirtualProduction/TextureShare/Source/TextureShare/Private/Object/TextureShareObjectProxy.h
@@ -28,7 +28,7 @@ public:
 	virtual const ITextureShareContext* GetTextureShareContext_RenderThread() const override;
 
 	virtual const FString& GetName_RenderThread() const override;
-	virtual const FTextureShareCoreObjectDesc& GetObjectDesc_RenderThread() const override;
+	virtual FTextureShareCoreObjectDesc GetObjectDesc_RenderThread() const override;
 
 	virtual bool IsActive_RenderThread() const override;
 	virtual bool IsFrameSyncActive_RenderThread() const override;
@@ -40,7 +40,7 @@ public:
 	virtual FTextureShareCoreProxyData& GetCoreProxyData_RenderThread() override;
 	virtual const FTextureShareCoreProxyData& GetCoreProxyData_RenderThread() const override;
 
-	virtual const TArray<FTextureShareCoreObjectProxyData>& GetReceivedCoreObjectProxyData_RenderThread() const override;
+	virtual TArray<FTextureShareCoreObjectProxyData> GetReceivedCoreObjectProxyData_RenderThread() const override;
 
 	virtual const FTextureShareData& GetData_RenderThread() const override;
 
diff --git a/Engine/Plugins/VirtualProduction/TextureShare/Source/TextureShareCore/Private/Containers/TextureShareCoreEnums.h b/Engine/Plugins/VirtualProduction/TextureShare/Source/TextureShareCore/Private/Containers/TextureShareCoreEnums.h
--- a/Engine/Plugins/VirtualProduction/TextureShare/Source/TextureShareCore/Private/Containers/TextureShareCoreEnums.h
+++ b/Engine/Plugins/VirtualProduction/TextureShare/Source/TextureShareCore/Private/Containers/TextureShareCoreEnums.h
@@ -58,8 +58,13 @@ enum class ETextureShareThreadMutex : uint8
 	RenderingThread,
 
 	InternalLock,
+	InternalObjectDescMT,
+	InternalFrameConnectionsMT,
+	InternalSyncSettingsMT,
 
-	COUNT
+	COUNT,
+	// All mutexes after this dont produce log
+	MAX_LOG = InternalLock
 };
 
 /**
diff --git a/Engine/Plugins/VirtualProduction/TextureShare/Source/TextureShareCore/Private/Core/TextureShareCoreHelpers.h b/Engine/Plugins/VirtualProduction/TextureShare/Source/TextureShareCore/Private/Core/TextureShareCoreHelpers.h
--- a/Engine/Plugins/VirtualProduction/TextureShare/Source/TextureShareCore/Private/Core/TextureShareCoreHelpers.h
+++ b/Engine/Plugins/VirtualProduction/TextureShare/Source/TextureShareCore/Private/Core/TextureShareCoreHelpers.h
@@ -47,6 +47,10 @@ namespace UE::TextureShareCore
 		case ETextureShareThreadMutex::GameThread:      return TEXT("GameThread");
 		case ETextureShareThreadMutex::RenderingThread: return TEXT("RenderingThread");
 		case ETextureShareThreadMutex::InternalLock:    return TEXT("InternalLock");
+
+		case ETextureShareThreadMutex::InternalObjectDescMT:       return TEXT("InternalObjectDescMT");
+		case ETextureShareThreadMutex::InternalFrameConnectionsMT: return TEXT("InternalFrameConnectionsMT");
+		case ETextureShareThreadMutex::InternalSyncSettingsMT:     return TEXT("InternalSyncSettingsMT");
 		default:
 			break;
 		}
@@ -171,13 +175,15 @@ namespace UE::TextureShareCore
 
 		return TEXT("?");
 	}
+	
+	static FString ToString(const FTextureShareCoreObjectSyncState& InObjectSyncState)
+	{
+		return FString::Printf(TEXT("%s[%s] {%s < %s}"), GetTEXT(InObjectSyncState.Step), GetTEXT(InObjectSyncState.State), GetTEXT(InObjectSyncState.NextStep), GetTEXT(InObjectSyncState.PrevStep));
+	}
 
 	static FString ToString(const FTextureShareCoreObjectDesc& InObject)
 	{
-		const FTextureShareCoreObjectSyncState& InSyncState = InObject.Sync.SyncState;
-		const FString SyncState = FString::Printf(TEXT("%s[%s]{%s < %s}"), GetTEXT(InSyncState.Step), GetTEXT(InSyncState.State), GetTEXT(InSyncState.NextStep), GetTEXT(InSyncState.PrevStep));
-
-		return FString::Printf(TEXT("%s(%s)"), *InObject.ProcessDesc.ProcessId, *SyncState);
+		return FString::Printf(TEXT("%s(%s)"), *InObject.ProcessDesc.ProcessId, *ToString(InObject.Sync.SyncState));
 	}
 
 	static FString ToString(const TArraySerializable<FTextureShareCoreObjectDesc>& InObjects)
diff --git a/Engine/Plugins/VirtualProduction/TextureShare/Source/TextureShareCore/Private/IPC/Containers/TextureShareCoreInterprocessObjectSyncState.cpp b/Engine/Plugins/VirtualProduction/TextureShare/Source/TextureShareCore/Private/IPC/Containers/TextureShareCoreInterprocessObjectSyncState.cpp
--- a/Engine/Plugins/VirtualProduction/TextureShare/Source/TextureShareCore/Private/IPC/Containers/TextureShareCoreInterprocessObjectSyncState.cpp
+++ b/Engine/Plugins/VirtualProduction/TextureShare/Source/TextureShareCore/Private/IPC/Containers/TextureShareCoreInterprocessObjectSyncState.cpp
@@ -23,7 +23,7 @@ ETextureShareInterprocessObjectSyncBarrierState FTextureShareCoreInterprocessObj
 
 	default:
 		// Add some log
-		UE_TS_LOG(LogTextureShareCoreObjectSync, Error, TEXT("BarrierState=%s : %s ask from %s"), GetTEXT(InBarrierState), *ToString(InSync), *ToString(*this));
+		UE_TS_LOG(LogTextureShareCoreObjectSync, Error, TEXT("\t\tBarrierState=%s : %s ask from %s"), GetTEXT(InBarrierState), *ToString(InSync), *ToString(*this));
 		break;
 	}
 
diff --git a/Engine/Plugins/VirtualProduction/TextureShare/Source/TextureShareCore/Private/IPC/TextureShareCoreInterprocessMutex.cpp b/Engine/Plugins/VirtualProduction/TextureShare/Source/TextureShareCore/Private/IPC/TextureShareCoreInterprocessMutex.cpp
--- a/Engine/Plugins/VirtualProduction/TextureShare/Source/TextureShareCore/Private/IPC/TextureShareCoreInterprocessMutex.cpp
+++ b/Engine/Plugins/VirtualProduction/TextureShare/Source/TextureShareCore/Private/IPC/TextureShareCoreInterprocessMutex.cpp
@@ -101,7 +101,6 @@ void FTextureShareCoreInterprocessMutex::UnlockMutex()
 			// relinquish
 			bLocked = false;
 
-			LOG_SCOPE_VERBOSITY_OVERRIDE(LogHAL, ELogVerbosity::Error);
 			ProcessMutex->Unlock();
 		}
 	}
diff --git a/Engine/Plugins/VirtualProduction/TextureShare/Source/TextureShareCore/Private/ITextureShareCoreObject.h b/Engine/Plugins/VirtualProduction/TextureShare/Source/TextureShareCore/Private/ITextureShareCoreObject.h
--- a/Engine/Plugins/VirtualProduction/TextureShare/Source/TextureShareCore/Private/ITextureShareCoreObject.h
+++ b/Engine/Plugins/VirtualProduction/TextureShare/Source/TextureShareCore/Private/ITextureShareCoreObject.h
@@ -23,12 +23,12 @@ public:
 	/**
 	* Returns detailed information about the TextureShare object.
 	*/
-	virtual const FTextureShareCoreObjectDesc& GetObjectDesc() const = 0;
+	virtual FTextureShareCoreObjectDesc GetObjectDesc() const = 0;
 
 	/**
 	* Returns detailed information about the TextureShare proxy object.
 	*/
-	virtual const FTextureShareCoreObjectDesc& GetObjectDesc_RenderThread() const = 0;
+	virtual FTextureShareCoreObjectDesc GetObjectDesc_RenderThread() const = 0;
 
 	/**
 	 * Returns true if the TextureShare object is ready to be used.
@@ -88,12 +88,12 @@ public:
 	 *
 	 * @return true, if success
 	 */
-	virtual bool SetSyncSetting(const FTextureShareCoreSyncSettings& InSyncSetting) = 0;
+	virtual bool SetSyncSettings(const FTextureShareCoreSyncSettings& InSyncSettings) = 0;
 
 	/**
 	* Return the sync settings of the TextureShare.
 	*/
-	virtual const FTextureShareCoreSyncSettings& GetSyncSetting() const = 0;
+	virtual FTextureShareCoreSyncSettings GetSyncSettings() const = 0;
 
 	/**
 	 * Get TextureShare default sync settings by template type
@@ -163,7 +163,7 @@ public:
 	 *
 	 * @return True if the success
 	 */
-	virtual bool LockThreadMutex(const ETextureShareThreadMutex InThreadMutex, bool bForceLockNoWait = false) = 0;
+	virtual bool LockThreadMutex(const ETextureShareThreadMutex InThreadMutex, bool bForceLockNoWait = false) const = 0;
 
 	/**
 	 * Support for multi-threaded implementation.
@@ -173,7 +173,7 @@ public:
 	 *
 	 * @return True if the success
 	 */
-	virtual bool UnlockThreadMutex(const ETextureShareThreadMutex InThreadMutex) = 0;
+	virtual bool UnlockThreadMutex(const ETextureShareThreadMutex InThreadMutex) const = 0;
 
 public:
 	//////////////////// Interprocess Synchronization ////////////////////
@@ -231,7 +231,7 @@ public:
 	 *
 	 * @return true, if success
 	 */
-	virtual const TArraySerializable<FTextureShareCoreObjectDesc>& GetConnectedInterprocessObjects() const = 0;
+	virtual TArraySerializable<FTextureShareCoreObjectDesc> GetConnectedInterprocessObjects() const = 0;
 
 public:
 	//////////////////// Data Containers ////////////////////
@@ -251,12 +251,12 @@ public:
 	/**
 	 * Received Data from connected process objects
 	 */
-	virtual const TArraySerializable<FTextureShareCoreObjectData>& GetReceivedData() const = 0;
+	virtual TArraySerializable<FTextureShareCoreObjectData> GetReceivedData() const = 0;
 
 	/**
 	 * Received ProxyData from connected process objects
 	 */
-	virtual const TArraySerializable<FTextureShareCoreObjectProxyData>& GetReceivedProxyData_RenderThread() const= 0;
+	virtual TArraySerializable<FTextureShareCoreObjectProxyData> GetReceivedProxyData_RenderThread() const= 0;
 
 public:
 	//////////////////// Destructor ////////////////////
diff --git a/Engine/Plugins/VirtualProduction/TextureShare/Source/TextureShareCore/Private/Object/TextureShareCoreObject.cpp b/Engine/Plugins/VirtualProduction/TextureShare/Source/TextureShareCore/Private/Object/TextureShareCoreObject.cpp
--- a/Engine/Plugins/VirtualProduction/TextureShare/Source/TextureShareCore/Private/Object/TextureShareCoreObject.cpp
+++ b/Engine/Plugins/VirtualProduction/TextureShare/Source/TextureShareCore/Private/Object/TextureShareCoreObject.cpp
@@ -18,24 +18,24 @@ namespace UE::TextureShareCore
 {
 	static FTextureShareCoreObjectDesc CreateNewObjectDesc(const FTextureShareCore& InOwner, const FString& InTextureShareName, const ETextureShareProcessType InProcessType)
 	{
-		FTextureShareCoreObjectDesc ObjectDesc;
+		FTextureShareCoreObjectDesc OutObjectDesc;
 
 		// The name of this object
-		ObjectDesc.ShareName = InTextureShareName;
+		OutObjectDesc.ShareName = InTextureShareName;
 
 		// Generate unique Guid for each object
-		ObjectDesc.ObjectGuid = FGuid::NewGuid();
+		OutObjectDesc.ObjectGuid = FGuid::NewGuid();
 
 		// Get value for Owner
-		ObjectDesc.ProcessDesc = InOwner.GetProcessDesc();
+		OutObjectDesc.ProcessDesc = InOwner.GetProcessDesc();
 
 		// Custom process type
 		if (InProcessType != ETextureShareProcessType::Undefined)
 		{
-			ObjectDesc.ProcessDesc.ProcessType = InProcessType;
+			OutObjectDesc.ProcessDesc.ProcessType = InProcessType;
 		}
 
-		return ObjectDesc;
+		return OutObjectDesc;
 	}
 };
 using namespace UE::TextureShareCore;
@@ -44,10 +44,10 @@ using namespace UE::TextureShareCore;
 // FTextureShareCoreObject
 //////////////////////////////////////////////////////////////////////////////////////////////
 FTextureShareCoreObject::FTextureShareCoreObject(FTextureShareCore& InOwner, const FString& InTextureShareName, const ETextureShareProcessType InProcessType)
-	: ObjectDesc(CreateNewObjectDesc(InOwner, InTextureShareName, InProcessType))
+	: ObjectDescMT(CreateNewObjectDesc(InOwner, InTextureShareName, InProcessType))
 	, Owner(InOwner)
 {
-	NotificationEvent = Owner.CreateInterprocessEvent(ObjectDesc.ObjectGuid);
+	NotificationEvent = Owner.CreateInterprocessEvent(ObjectDescMT.ObjectGuid);
 }
 
 FTextureShareCoreObject::~FTextureShareCoreObject()
@@ -61,8 +61,8 @@ FTextureShareCoreObject::~FTextureShareCoreObject()
 void FTextureShareCoreObject::HandleResetSync(FTextureShareCoreInterprocessMemory& InterprocessMemory, FTextureShareCoreInterprocessObject& LocalObject)
 {
 #if TEXTURESHARECORE_DEBUGLOG
-	InterprocessMemory.UpdateFrameConnections(FrameConnections);
-	UE_TS_LOG(LogTextureShareCoreObjectSync, Error, TEXT("%s:ResetSync(%s) %s"), *GetName(), *ToString(LocalObject), *ToString(FrameConnections));
+	UpdateFrameConnections(InterprocessMemory);
+	UE_TS_LOG(LogTextureShareCoreObjectSync, Error, TEXT("%s:ResetSync(%s) %s"), *GetName(), *ToString(LocalObject), *ToString(GetFrameConnections()));
 #endif
 
 	// Unlock multithread thread mutexes
@@ -81,7 +81,7 @@ void FTextureShareCoreObject::HandleResetSync(FTextureShareCoreInterprocessMemor
 	SetFrameSyncState(ETextureShareCoreInterprocessObjectFrameSyncState::FrameSyncLost);
 
 	// Reset current frame connections
-	FrameConnections.Reset();
+	ResetFrameConnections();
 }
 
 //////////////////////////////////////////////////////////////////////////////////////////////
@@ -89,7 +89,7 @@ void FTextureShareCoreObject::SendNotificationEvents(const bool bInterprocessMem
 {
 	if (IsSessionActive() && IsActive())
 	{
-		if (!bInterprocessMemoryLockRequired || Owner.LockInterprocessMemory(SyncSettings.TimeoutSettings.MemoryMutexTimeout))
+		if (!bInterprocessMemoryLockRequired || Owner.LockInterprocessMemory(GetTimeOutSettings().MemoryMutexTimeout))
 		{
 			if (FTextureShareCoreInterprocessMemory* InterprocessMemory = Owner.GetInterprocessMemory())
 			{
@@ -167,7 +167,7 @@ bool FTextureShareCoreObject::TryWaitFrameProcesses(const uint32 InRemainMaxMill
 	NotificationEvent->Wait(InRemainMaxMillisecondsToWait);
 
 	// Try to lock shared memory again
-	return Owner.LockInterprocessMemory(SyncSettings.TimeoutSettings.MemoryMutexTimeout);
+	return Owner.LockInterprocessMemory(GetTimeOutSettings().MemoryMutexTimeout);
 }
 
 bool FTextureShareCoreObject::RemoveObject()
diff --git a/Engine/Plugins/VirtualProduction/TextureShare/Source/TextureShareCore/Private/Object/TextureShareCoreObject.h b/Engine/Plugins/VirtualProduction/TextureShare/Source/TextureShareCore/Private/Object/TextureShareCoreObject.h
--- a/Engine/Plugins/VirtualProduction/TextureShare/Source/TextureShareCore/Private/Object/TextureShareCoreObject.h
+++ b/Engine/Plugins/VirtualProduction/TextureShare/Source/TextureShareCore/Private/Object/TextureShareCoreObject.h
@@ -4,6 +4,7 @@
 #include "ITextureShareCoreObject.h"
 #include "Core/TextureShareCore.h"
 
+#include "IPC/Containers/TextureShareCoreInterprocessMemory.h"
 #include "IPC/Containers/TextureShareCoreInterprocessObject.h"
 #include "IPC/Containers/TextureShareCoreInterprocessEnums.h"
 
@@ -22,8 +23,8 @@ public:
 
 	virtual const FString& GetName() const override;
 
-	virtual const FTextureShareCoreObjectDesc& GetObjectDesc() const override;
-	virtual const FTextureShareCoreObjectDesc& GetObjectDesc_RenderThread() const override;
+	virtual FTextureShareCoreObjectDesc GetObjectDesc() const override;
+	virtual FTextureShareCoreObjectDesc GetObjectDesc_RenderThread() const override;
 
 	virtual bool IsActive() const override;
 	virtual bool IsActive_RenderThread() const override;
@@ -34,8 +35,12 @@ public:
 	virtual bool IsBeginFrameSyncActive() const override;
 	virtual bool IsBeginFrameSyncActive_RenderThread() const override;
 
+protected:
+	const FTextureShareCoreObjectDesc GetObjectDesc_AnyThread() const;
+
 private:
-	FTextureShareCoreObjectDesc ObjectDesc;
+	// (MultiThread access data)
+	FTextureShareCoreObjectDesc ObjectDescMT;
 
 public:
 	///////////////////// Settings /////////////////////
@@ -43,19 +48,23 @@ public:
 	virtual bool SetProcessId(const FString& InProcessId) override;
 	virtual bool SetDeviceType(const ETextureShareDeviceType InDeviceType) override;
 
-	virtual bool SetSyncSetting(const FTextureShareCoreSyncSettings& InSyncSettings) override;
-	virtual const FTextureShareCoreSyncSettings& GetSyncSetting() const override;
+	virtual bool SetSyncSettings(const FTextureShareCoreSyncSettings& InSyncSettings) override;
+	virtual FTextureShareCoreSyncSettings GetSyncSettings() const override;
 
 	virtual FTextureShareCoreFrameSyncSettings GetFrameSyncSettings(const ETextureShareFrameSyncTemplate InType) const override;
 
 private:
 	// This settings used for a frame
-	FTextureShareCoreSyncSettings SyncSettings;
+	// (MultiThread access data)
+	FTextureShareCoreSyncSettings SyncSettingsMT;
 
 protected:
 	void AddNewSyncStep(const ETextureShareSyncStep InSyncStep);
 	void UpdateInterprocessObject();
 
+	/** Special exception for timeout settings: they do not use arrays and are thread-safe. */
+	const FTextureShareCoreTimeOutSettings& GetTimeOutSettings() const;
+
 public:
 	///////////////////// Session /////////////////////
 
@@ -69,11 +78,11 @@ private:
 public:
 	///////////////////// Thread sync support /////////////////////
 
-	virtual bool LockThreadMutex(const ETextureShareThreadMutex InThreadMutex, bool bForceLockNoWait = false) override;
-	virtual bool UnlockThreadMutex(const ETextureShareThreadMutex InThreadMutex) override;
+	virtual bool LockThreadMutex(const ETextureShareThreadMutex InThreadMutex, bool bForceLockNoWait = false) const override;
+	virtual bool UnlockThreadMutex(const ETextureShareThreadMutex InThreadMutex) const override;
 
 protected:
-	TSharedPtr<FTextureShareCoreInterprocessMutex, ESPMode::ThreadSafe> GetThreadMutex(const ETextureShareThreadMutex InThreadMutex);
+	TSharedPtr<FTextureShareCoreInterprocessMutex, ESPMode::ThreadSafe> GetThreadMutex(const ETextureShareThreadMutex InThreadMutex) const;
 	void ResetThreadMutex(const ETextureShareThreadMutex InThreadMutex);
 
 	void InitializeThreadMutexes();
@@ -96,7 +105,7 @@ public:
 	virtual bool FindSkippedSyncStep(const ETextureShareSyncStep InSyncStep, ETextureShareSyncStep& OutSkippedSyncStep) override;
 	virtual bool FindSkippedSyncStep_RenderThread(const ETextureShareSyncStep InSyncStep, ETextureShareSyncStep& OutSkippedSyncStep) override;
 
-	virtual const TArraySerializable<FTextureShareCoreObjectDesc>& GetConnectedInterprocessObjects() const override;
+	virtual TArraySerializable<FTextureShareCoreObjectDesc> GetConnectedInterprocessObjects() const override;
 
 protected:
 	bool FindSkippedSyncStepImpl(const ETextureShareSyncStep InSyncStep, bool bIsProxyFrame, ETextureShareSyncStep& OutSkippedSyncStep);
@@ -105,7 +114,8 @@ protected:
 private:
 	// When the frame starts, a list of processes connected to the frame is created.
 	// New processes connected later inside this framelock will be ignored.
-	TArraySerializable<FTextureShareCoreObjectDesc> FrameConnections;
+	// (MultiThread access data)
+	TArraySerializable<FTextureShareCoreObjectDesc> FrameConnectionsMT;
 
 	// Prevend double check for begion connection timeout
 	bool bIsFrameConnectionTimeoutReached = false;
@@ -124,6 +134,24 @@ protected:
 	bool ConnectFrameProcesses();
 	bool DisconnectFrameProcesses();
 
+	/** Return FrameConnectionsMT. */
+	const TArraySerializable<FTextureShareCoreObjectDesc> GetFrameConnections() const;
+
+	/** Return true is FrameConnectionsMT is empty. */
+	bool IsEmptyFrameConnections() const;
+
+	/** Reset FrameConnectionsMT array. */
+	void ResetFrameConnections();
+
+	/** Set FrameConnectionsMT array. */
+	void SetFrameConnections(const TArraySerializable<FTextureShareCoreObjectDesc>& InNewFrameConnections);
+
+	/** Update frame connections data from the IPC memory. */
+	void UpdateFrameConnections(FTextureShareCoreInterprocessMemory& InterprocessMemory);
+
+	/** Search for a frame connections. */
+	int32 FindFrameConnections(FTextureShareCoreInterprocessMemory& InterprocessMemory, FTextureShareCoreInterprocessObject& LocalObject);
+
 	bool DoFrameSync(const ETextureShareSyncStep InSyncStep);
 	bool DoFrameSync_RenderThread(const ETextureShareSyncStep InSyncStep);
 
@@ -163,8 +191,8 @@ public:
 	virtual FTextureShareCoreData&      GetData() override;
 	virtual FTextureShareCoreProxyData& GetProxyData_RenderThread() override;
 
-	virtual const TArraySerializable<FTextureShareCoreObjectData>&      GetReceivedData() const override;
-	virtual const TArraySerializable<FTextureShareCoreObjectProxyData>& GetReceivedProxyData_RenderThread() const override;
+	virtual TArraySerializable<FTextureShareCoreObjectData>      GetReceivedData() const override;
+	virtual TArraySerializable<FTextureShareCoreObjectProxyData> GetReceivedProxyData_RenderThread() const override;
 
 protected:
 	void ReleaseData();
diff --git a/Engine/Plugins/VirtualProduction/TextureShare/Source/TextureShareCore/Private/Object/TextureShareCoreObject_Connect.cpp b/Engine/Plugins/VirtualProduction/TextureShare/Source/TextureShareCore/Private/Object/TextureShareCoreObject_Connect.cpp
--- a/Engine/Plugins/VirtualProduction/TextureShare/Source/TextureShareCore/Private/Object/TextureShareCoreObject_Connect.cpp
+++ b/Engine/Plugins/VirtualProduction/TextureShare/Source/TextureShareCore/Private/Object/TextureShareCoreObject_Connect.cpp
@@ -27,12 +27,12 @@ bool FTextureShareCoreObject::TryFrameProcessesConnection(FTextureShareCoreInter
 	LocalObject.Sync.UpdateLastAccessTime();
 
 	// Collect valid processes to connect
-	const int32 ReadyToConnectObjectsCount = InterprocessMemory.FindConnectableObjects(FrameConnections, LocalObject);
+	const int32 ReadyToConnectObjectsCount = FindFrameConnections(InterprocessMemory, LocalObject);
 
-	UE_TS_BARRIER_LOG(LogTextureShareCoreObjectSync, Log, TEXT("%s:TryFrameProcessesConnection(%s)=%d %s"), *GetName(), *ToString(LocalObject), ReadyToConnectObjectsCount, *ToString(FrameConnections));
+	UE_TS_BARRIER_LOG(LogTextureShareCoreObjectSync, Log, TEXT("%s:TryFrameProcessesConnection(%s)=%d %s"), *GetName(), *ToString(LocalObject), ReadyToConnectObjectsCount, *ToString(GetFrameConnections()));
 
 	// Handle sync logic
-	const ETextureShareCoreFrameConnectionsState ConnectionsState = SyncSettings.FrameConnectionSettings.GetConnectionsState(ReadyToConnectObjectsCount, FrameConnections.Num());
+	const ETextureShareCoreFrameConnectionsState ConnectionsState = GetSyncSettings().FrameConnectionSettings.GetConnectionsState(ReadyToConnectObjectsCount, GetFrameConnections().Num());
 	switch (ConnectionsState)
 	{
 	case ETextureShareCoreFrameConnectionsState::SkipFrame:
@@ -78,7 +78,7 @@ bool FTextureShareCoreObject::TryFrameProcessesConnection(FTextureShareCoreInter
 	}
 
 	// Reset connections list
-	FrameConnections.Reset();
+	ResetFrameConnections();
 
 	// Wait for new frame connection
 	return true;
@@ -88,9 +88,9 @@ bool FTextureShareCoreObject::ConnectFrameProcesses()
 {
 	UE_TS_BARRIER_LOG(LogTextureShareCoreObjectSync, Log, TEXT("%s:ConnectFrameProcesses()"), *GetName());
 
-	FrameConnections.Reset();
+	ResetFrameConnections();
 
-	if (IsSessionActive() && IsActive() && Owner.LockInterprocessMemory(SyncSettings.TimeoutSettings.MemoryMutexTimeout))
+	if (IsSessionActive() && IsActive() && Owner.LockInterprocessMemory(GetTimeOutSettings().MemoryMutexTimeout))
 	{
 		if (FTextureShareCoreInterprocessMemory* InterprocessMemory = Owner.GetInterprocessMemory())
 		{
@@ -103,7 +103,7 @@ bool FTextureShareCoreObject::ConnectFrameProcesses()
 				if (BeginSyncBarrier(*InterprocessMemory, *LocalObject, ETextureShareSyncStep::InterprocessConnection, ETextureShareSyncPass::Enter))
 				{
 					// Begin connect logic:
-					FTextureShareCoreObjectTimeout FrameBeginTimer(SyncSettings.TimeoutSettings.FrameBeginTimeOut, SyncSettings.TimeoutSettings.FrameBeginTimeOutSplit);
+					FTextureShareCoreObjectTimeout FrameBeginTimer(GetTimeOutSettings().FrameBeginTimeOut, GetTimeOutSettings().FrameBeginTimeOutSplit);
 					while (TryFrameProcessesConnection(*InterprocessMemory, *LocalObject))
 					{
 						if (FrameBeginTimer.IsTimeOut())
@@ -126,7 +126,7 @@ bool FTextureShareCoreObject::ConnectFrameProcesses()
 				}
 
 				// There is no process to connect, reset sync
-				if (FrameConnections.IsEmpty() && FrameSyncState == ETextureShareCoreInterprocessObjectFrameSyncState::NewFrame)
+				if (IsEmptyFrameConnections() && FrameSyncState == ETextureShareCoreInterprocessObjectFrameSyncState::NewFrame)
 				{
 					// No process to connect, let's set the synchronization state as 'FrameSyncLost'.
 					// This state is used in functions IsBeginFrameSyncActive() and IsBeginFrameSyncActive_RenderThread()
@@ -149,7 +149,7 @@ bool FTextureShareCoreObject::ConnectFrameProcesses()
 		Owner.UnlockInterprocessMemory();
 	}
 
-	if (!FrameConnections.IsEmpty())
+	if (!IsEmptyFrameConnections())
 	{
 
 		// Wait for other processes finish frame connect
diff --git a/Engine/Plugins/VirtualProduction/TextureShare/Source/TextureShareCore/Private/Object/TextureShareCoreObject_Data.cpp b/Engine/Plugins/VirtualProduction/TextureShare/Source/TextureShareCore/Private/Object/TextureShareCoreObject_Data.cpp
--- a/Engine/Plugins/VirtualProduction/TextureShare/Source/TextureShareCore/Private/Object/TextureShareCoreObject_Data.cpp
+++ b/Engine/Plugins/VirtualProduction/TextureShare/Source/TextureShareCore/Private/Object/TextureShareCoreObject_Data.cpp
@@ -30,12 +30,12 @@ FTextureShareCoreProxyData& FTextureShareCoreObject::GetProxyData_RenderThread()
 	return ProxyData;
 }
 
-const TArraySerializable<FTextureShareCoreObjectData>& FTextureShareCoreObject::GetReceivedData() const
+TArraySerializable<FTextureShareCoreObjectData> FTextureShareCoreObject::GetReceivedData() const
 {
 	return ReceivedObjectsData;
 }
 
-const TArraySerializable<FTextureShareCoreObjectProxyData>& FTextureShareCoreObject::GetReceivedProxyData_RenderThread() const
+TArraySerializable<FTextureShareCoreObjectProxyData> FTextureShareCoreObject::GetReceivedProxyData_RenderThread() const
 {
 	return ReceivedObjectsProxyData;
 }
@@ -44,7 +44,7 @@ const TArraySerializable<FTextureShareCoreObjectProxyData>& FTextureShareCoreObj
 bool FTextureShareCoreObject::SendFrameData()
 {
 	bool bResult = false;
-	if (IsFrameSyncActive() && Owner.LockInterprocessMemory(SyncSettings.TimeoutSettings.MemoryMutexTimeout))
+	if (IsFrameSyncActive() && Owner.LockInterprocessMemory(GetTimeOutSettings().MemoryMutexTimeout))
 	{
 		if (FTextureShareCoreInterprocessMemory* InterprocessMemory = Owner.GetInterprocessMemory())
 		{
@@ -68,13 +68,13 @@ bool FTextureShareCoreObject::ReceiveFrameData()
 	ReceivedObjectsData.Empty();
 
 	bool bResult = false;
-	if (IsFrameSyncActive() && Owner.LockInterprocessMemory(SyncSettings.TimeoutSettings.MemoryMutexTimeout))
+	if (IsFrameSyncActive() && Owner.LockInterprocessMemory(GetTimeOutSettings().MemoryMutexTimeout))
 	{
 		if (FTextureShareCoreInterprocessMemory* InterprocessMemory = Owner.GetInterprocessMemory())
 		{
 			bResult = true;
 			// Read data from remote processes
-			for (const FTextureShareCoreObjectDesc& RemoteObjDescIt : FrameConnections)
+			for (const FTextureShareCoreObjectDesc& RemoteObjDescIt : GetFrameConnections())
 			{
 				if (FTextureShareCoreInterprocessObject* InterprocessObject = InterprocessMemory->FindObject(RemoteObjDescIt))
 				{
@@ -101,7 +101,7 @@ bool FTextureShareCoreObject::ReceiveFrameData()
 bool FTextureShareCoreObject::SendFrameProxyData_RenderThread()
 {
 	bool bResult = false;
-	if (IsFrameSyncActive_RenderThread() && Owner.LockInterprocessMemory(SyncSettings.TimeoutSettings.MemoryMutexTimeout))
+	if (IsFrameSyncActive_RenderThread() && Owner.LockInterprocessMemory(GetTimeOutSettings().MemoryMutexTimeout))
 	{
 		if (FTextureShareCoreInterprocessMemory* InterprocessMemory = Owner.GetInterprocessMemory())
 		{
@@ -125,13 +125,13 @@ bool FTextureShareCoreObject::ReceiveFrameProxyData_RenderThread()
 	ReceivedObjectsProxyData.Empty();
 
 	bool bResult = false;
-	if (IsFrameSyncActive_RenderThread() && Owner.LockInterprocessMemory(SyncSettings.TimeoutSettings.MemoryMutexTimeout))
+	if (IsFrameSyncActive_RenderThread() && Owner.LockInterprocessMemory(GetTimeOutSettings().MemoryMutexTimeout))
 	{
 		if (FTextureShareCoreInterprocessMemory* InterprocessMemory = Owner.GetInterprocessMemory())
 		{
 			bResult = true;
 			// Read data from remote processes
-			for (const FTextureShareCoreObjectDesc& RemoteObjDescIt : FrameConnections)
+			for (const FTextureShareCoreObjectDesc& RemoteObjDescIt : GetFrameConnections())
 			{
 				if (FTextureShareCoreInterprocessObject* InterprocessObject = InterprocessMemory->FindObject(RemoteObjDescIt))
 				{
diff --git a/Engine/Plugins/VirtualProduction/TextureShare/Source/TextureShareCore/Private/Object/TextureShareCoreObject_Session.cpp b/Engine/Plugins/VirtualProduction/TextureShare/Source/TextureShareCore/Private/Object/TextureShareCoreObject_Session.cpp
--- a/Engine/Plugins/VirtualProduction/TextureShare/Source/TextureShareCore/Private/Object/TextureShareCoreObject_Session.cpp
+++ b/Engine/Plugins/VirtualProduction/TextureShare/Source/TextureShareCore/Private/Object/TextureShareCoreObject_Session.cpp
@@ -23,10 +23,12 @@ using namespace UE::TextureShareCore;
 //////////////////////////////////////////////////////////////////////////////////////////////
 bool FTextureShareCoreObject::BeginSession()
 {
-	if (!IsSessionActive() && IsActive() && Owner.LockInterprocessMemory(SyncSettings.TimeoutSettings.MemoryMutexTimeout))
+	if (!IsSessionActive() && IsActive() && Owner.LockInterprocessMemory(GetTimeOutSettings().MemoryMutexTimeout))
 	{
+		const FTextureShareCoreSyncSettings CurrentSyncSettings = GetSyncSettings();
+
 		// Update object desc sync hash value
-		ObjectDesc.Sync.SetSyncStepSettings(SyncSettings);
+		ObjectDescMT.Sync.SetSyncStepSettings(CurrentSyncSettings);
 
 		if (FTextureShareCoreInterprocessMemory* InterprocessMemory = Owner.GetInterprocessMemory())
 		{
@@ -37,7 +39,7 @@ bool FTextureShareCoreObject::BeginSession()
 			if (FTextureShareCoreInterprocessObject* InterprocessObject = InterprocessMemory->FindEmptyObject())
 			{
 				// Initialize IPC object data
-				InterprocessObject->InitializeInterprocessObject(GetObjectDesc(), SyncSettings);
+				InterprocessObject->InitializeInterprocessObject(GetObjectDesc(), CurrentSyncSettings);
 
 				InitializeThreadMutexes();
 
@@ -68,7 +70,7 @@ bool FTextureShareCoreObject::EndSession()
 		bSessionActive = false;
 
 		// Release IPC
-		if (Owner.LockInterprocessMemory(SyncSettings.TimeoutSettings.MemoryMutexTimeout))
+		if (Owner.LockInterprocessMemory(GetTimeOutSettings().MemoryMutexTimeout))
 		{
 			if (FTextureShareCoreInterprocessMemory* InterprocessMemory = Owner.GetInterprocessMemory())
 			{
diff --git a/Engine/Plugins/VirtualProduction/TextureShare/Source/TextureShareCore/Private/Object/TextureShareCoreObject_Settings.cpp b/Engine/Plugins/VirtualProduction/TextureShare/Source/TextureShareCore/Private/Object/TextureShareCoreObject_Settings.cpp
--- a/Engine/Plugins/VirtualProduction/TextureShare/Source/TextureShareCore/Private/Object/TextureShareCoreObject_Settings.cpp
+++ b/Engine/Plugins/VirtualProduction/TextureShare/Source/TextureShareCore/Private/Object/TextureShareCoreObject_Settings.cpp
@@ -21,127 +21,185 @@ using namespace UE::TextureShareCore;
 //////////////////////////////////////////////////////////////////////////////////////////////
 bool FTextureShareCoreObject::SetProcessId(const FString& InProcessId)
 {
-	if (InProcessId.IsEmpty() || ObjectDesc.ProcessDesc.ProcessId == InProcessId)
+	if (InProcessId.IsEmpty())
 	{
 		// Skip equal or empty values
 		return false;
 	}
 
+	bool bShouldSetProcessId = false;
 	if (!IsSessionActive())
 	{
-		// Support local process name changes on the fly
-		ObjectDesc.ProcessDesc.ProcessId = InProcessId;
+		bShouldSetProcessId = (ObjectDescMT.ProcessDesc.ProcessId != InProcessId);
+		if (bShouldSetProcessId)
+		{
+			ObjectDescMT.ProcessDesc.ProcessId = InProcessId;
+		}
 	}
-	else
+	else if (LockThreadMutex(ETextureShareThreadMutex::InternalObjectDescMT))
 	{
-		if (LockThreadMutex(ETextureShareThreadMutex::InternalLock))
+		bShouldSetProcessId = (ObjectDescMT.ProcessDesc.ProcessId != InProcessId);
+		if (bShouldSetProcessId)
 		{
-			ObjectDesc.ProcessDesc.ProcessId = InProcessId;
-
-			// Update IPC
-			UpdateInterprocessObject();
-
-			UnlockThreadMutex(ETextureShareThreadMutex::InternalLock);
+			ObjectDescMT.ProcessDesc.ProcessId = InProcessId;
 		}
+		UnlockThreadMutex(ETextureShareThreadMutex::InternalObjectDescMT);
+	}
+
+	// Update IPC
+	if (bShouldSetProcessId)
+	{
+		UpdateInterprocessObject();
 	}
 
-	UE_TS_LOG(LogTextureShareCoreObject, Log, TEXT("%s:SetProcessId('%s')"), *GetName(), *InProcessId);
+	if (bShouldSetProcessId)
+	{
+		UE_TS_LOG(LogTextureShareCoreObject, Log, TEXT("%s:SetProcessId('%s')"), *GetName(), *InProcessId);
+	}
 
 	return true;
 }
 
 bool FTextureShareCoreObject::SetDeviceType(const ETextureShareDeviceType InDeviceType)
 {
+	bool bShouldSetDeviceType = false;
 	if (!IsSessionActive())
 	{
-		ObjectDesc.ProcessDesc.DeviceType = InDeviceType;
+		bShouldSetDeviceType = (ObjectDescMT.ProcessDesc.DeviceType != InDeviceType);
+		if (bShouldSetDeviceType)
+		{
+			ObjectDescMT.ProcessDesc.DeviceType = InDeviceType;
+		}
 	}
-	else
+	else if (LockThreadMutex(ETextureShareThreadMutex::InternalObjectDescMT))
 	{
-		if (LockThreadMutex(ETextureShareThreadMutex::InternalLock))
+		bShouldSetDeviceType = (ObjectDescMT.ProcessDesc.DeviceType != InDeviceType);
+		if (bShouldSetDeviceType)
 		{
-			ObjectDesc.ProcessDesc.DeviceType = InDeviceType;
-
-			// Update IPC
-			UpdateInterprocessObject();
-
-			UnlockThreadMutex(ETextureShareThreadMutex::InternalLock);
+			ObjectDescMT.ProcessDesc.DeviceType = InDeviceType;
 		}
+		UnlockThreadMutex(ETextureShareThreadMutex::InternalObjectDescMT);
+	}
+
+	// Update IPC
+	if (bShouldSetDeviceType)
+	{
+		UpdateInterprocessObject();
 	}
 
-	UE_TS_LOG(LogTextureShareCoreObject, Log, TEXT("%s:SetDeviceType(%s)"), *GetName(), GetTEXT(InDeviceType));
+	if (bShouldSetDeviceType)
+	{
+		UE_TS_LOG(LogTextureShareCoreObject, Log, TEXT("%s:SetDeviceType(%s)"), *GetName(), GetTEXT(InDeviceType));
+	}
 
 	return true;
 }
 
 //////////////////////////////////////////////////////////////////////////////////////////////
-bool FTextureShareCoreObject::SetSyncSetting(const FTextureShareCoreSyncSettings& InSyncSettings)
+const FTextureShareCoreTimeOutSettings& FTextureShareCoreObject::GetTimeOutSettings() const
 {
+	// Special exception for timeout settings: they do not use arrays and are thread-safe.
+	return SyncSettingsMT.TimeoutSettings;
+}
+
+FTextureShareCoreSyncSettings FTextureShareCoreObject::GetSyncSettings() const
+{
+	FTextureShareCoreSyncSettings OutSyncSettings;
 	if (!IsSessionActive())
 	{
-		SyncSettings = InSyncSettings;
+		OutSyncSettings = SyncSettingsMT;
 	}
-	else
+	else if (LockThreadMutex(ETextureShareThreadMutex::InternalSyncSettingsMT))
 	{
-		if (SyncSettings ==InSyncSettings)
-		{
-			return true;
-		}
+		OutSyncSettings = SyncSettingsMT;
+		UnlockThreadMutex(ETextureShareThreadMutex::InternalSyncSettingsMT);
+	}
 
-		if (LockThreadMutex(ETextureShareThreadMutex::InternalLock))
-		{
-			// Update sync settings
-			SyncSettings = InSyncSettings;
+	return OutSyncSettings;
+}
 
-			// Update object desc sync hash value
-			ObjectDesc.Sync.SetSyncStepSettings(SyncSettings);
+bool FTextureShareCoreObject::SetSyncSettings(const FTextureShareCoreSyncSettings& InSyncSettings)
+{
+	bool bShouldUpdateSyncSettings = false;
+	if (!IsSessionActive())
+	{
+		bShouldUpdateSyncSettings = SyncSettingsMT != InSyncSettings;
+		SyncSettingsMT = InSyncSettings;
+	}
+	else if (LockThreadMutex(ETextureShareThreadMutex::InternalSyncSettingsMT))
+	{
+		bShouldUpdateSyncSettings = SyncSettingsMT != InSyncSettings;
+		if (bShouldUpdateSyncSettings)
+		{
+			UE_TS_LOG(LogTextureShareCoreObject, Log, TEXT("%s:SetSyncSettings()"), *GetName());
 
-			// Update IPC
-			UpdateInterprocessObject();
+			SyncSettingsMT = InSyncSettings;
 
-			UnlockThreadMutex(ETextureShareThreadMutex::InternalLock);
+			if (LockThreadMutex(ETextureShareThreadMutex::InternalObjectDescMT))
+			{
+				// Update object desc sync hash value
+				ObjectDescMT.Sync.SetSyncStepSettings(SyncSettingsMT);
+				UnlockThreadMutex(ETextureShareThreadMutex::InternalObjectDescMT);
+			}
 		}
+
+		UnlockThreadMutex(ETextureShareThreadMutex::InternalSyncSettingsMT);
 	}
 
-	UE_TS_LOG(LogTextureShareCoreObject, Log, TEXT("%s:SetSyncSetting()"), *GetName());
+	// Update IPC
+	if (bShouldUpdateSyncSettings)
+	{
+		UpdateInterprocessObject();
+	}
 
 	return true;
 }
 
 void FTextureShareCoreObject::AddNewSyncStep(const ETextureShareSyncStep InSyncStep)
 {
+	bool bShouldUpdateSyncSettings = false;
 	if (!IsSessionActive())
 	{
-		SyncSettings.FrameSyncSettings.Steps.AddSorted(InSyncStep);
+		bShouldUpdateSyncSettings = !SyncSettingsMT.FrameSyncSettings.Steps.Contains(InSyncStep);
+		SyncSettingsMT.FrameSyncSettings.Steps.AddSorted(InSyncStep);
 	}
-	else
+	else if (LockThreadMutex(ETextureShareThreadMutex::InternalSyncSettingsMT))
 	{
-		if (LockThreadMutex(ETextureShareThreadMutex::InternalLock))
+		bShouldUpdateSyncSettings = !SyncSettingsMT.FrameSyncSettings.Steps.Contains(InSyncStep);
+		if (bShouldUpdateSyncSettings)
 		{
-			// Add requested sync step
-			SyncSettings.FrameSyncSettings.Steps.AddSorted(InSyncStep);
-
-			// Update object desc sync hash value
-			ObjectDesc.Sync.SetSyncStepSettings(SyncSettings);
+			UE_TS_LOG(LogTextureShareCoreObject, Log, TEXT("%s:AddNewSyncStep(%s)"), *GetName(), GetTEXT(InSyncStep));
 
-			// Update IPC
-			UpdateInterprocessObject();
+			// Add requested sync step
+			SyncSettingsMT.FrameSyncSettings.Steps.AddSorted(InSyncStep);
 
-			UnlockThreadMutex(ETextureShareThreadMutex::InternalLock);
+			if (LockThreadMutex(ETextureShareThreadMutex::InternalObjectDescMT))
+			{
+				// Update object desc sync hash value
+				ObjectDescMT.Sync.SetSyncStepSettings(SyncSettingsMT);
+				UnlockThreadMutex(ETextureShareThreadMutex::InternalObjectDescMT);
+			}
 		}
+		UnlockThreadMutex(ETextureShareThreadMutex::InternalSyncSettingsMT);
+	}
+
+	// Update IPC
+	if (bShouldUpdateSyncSettings)
+	{
+		UpdateInterprocessObject();
 	}
 }
 
 void FTextureShareCoreObject::UpdateInterprocessObject()
 {
 	// Also update settings exposition for other processes
-	if (IsActive() && Owner.LockInterprocessMemory(SyncSettings.TimeoutSettings.MemoryMutexTimeout))
+	if (IsSessionActive() && IsActive() && Owner.LockInterprocessMemory(GetTimeOutSettings().MemoryMutexTimeout))
 	{
 		if (FTextureShareCoreInterprocessMemory* InterprocessMemory = Owner.GetInterprocessMemory())
 		{
 			if (FTextureShareCoreInterprocessObject* LocalObject = InterprocessMemory->FindObject(GetObjectDesc()))
 			{
-				LocalObject->UpdateInterprocessObject(GetObjectDesc(), SyncSettings);
+				LocalObject->UpdateInterprocessObject(GetObjectDesc(), GetSyncSettings());
 			}
 		}
 
@@ -151,11 +209,6 @@ void FTextureShareCoreObject::UpdateInterprocessObject()
 	UE_TS_LOG(LogTextureShareCoreObject, Log, TEXT("%s:UpdateInterprocessObject()"), *GetName());
 }
 
-const FTextureShareCoreSyncSettings& FTextureShareCoreObject::GetSyncSetting() const
-{
-	return SyncSettings;
-}
-
 FTextureShareCoreFrameSyncSettings FTextureShareCoreObject::GetFrameSyncSettings(const ETextureShareFrameSyncTemplate InType) const
 {
 	FTextureShareCoreFrameSyncSettings Result;
diff --git a/Engine/Plugins/VirtualProduction/TextureShare/Source/TextureShareCore/Private/Object/TextureShareCoreObject_State.cpp b/Engine/Plugins/VirtualProduction/TextureShare/Source/TextureShareCore/Private/Object/TextureShareCoreObject_State.cpp
--- a/Engine/Plugins/VirtualProduction/TextureShare/Source/TextureShareCore/Private/Object/TextureShareCoreObject_State.cpp
+++ b/Engine/Plugins/VirtualProduction/TextureShare/Source/TextureShareCore/Private/Object/TextureShareCoreObject_State.cpp
@@ -21,17 +21,33 @@ using namespace UE::TextureShareCore;
 //////////////////////////////////////////////////////////////////////////////////////////////
 const FString& FTextureShareCoreObject::GetName() const
 {
-	return ObjectDesc.ShareName;
+	return ObjectDescMT.ShareName;
 }
 
-const FTextureShareCoreObjectDesc& FTextureShareCoreObject::GetObjectDesc() const
+const FTextureShareCoreObjectDesc FTextureShareCoreObject::GetObjectDesc_AnyThread() const
 {
-	return ObjectDesc;
+	FTextureShareCoreObjectDesc OutObjectDesc;
+	if (!IsSessionActive())
+	{
+		OutObjectDesc = ObjectDescMT;
+	}
+	else if (LockThreadMutex(ETextureShareThreadMutex::InternalObjectDescMT))
+	{
+		OutObjectDesc = ObjectDescMT;
+		UnlockThreadMutex(ETextureShareThreadMutex::InternalObjectDescMT);
+	}
+
+	return OutObjectDesc;
+}
+
+FTextureShareCoreObjectDesc FTextureShareCoreObject::GetObjectDesc() const
+{
+	return GetObjectDesc_AnyThread();
 }
 
-const FTextureShareCoreObjectDesc& FTextureShareCoreObject::GetObjectDesc_RenderThread() const
+FTextureShareCoreObjectDesc FTextureShareCoreObject::GetObjectDesc_RenderThread() const
 {
-	return ObjectDesc;
+	return GetObjectDesc_AnyThread();
 }
 //////////////////////////////////////////////////////////////////////////////////////////////
 bool FTextureShareCoreObject::IsActive() const
@@ -46,12 +62,12 @@ bool FTextureShareCoreObject::IsActive_RenderThread() const
 
 bool FTextureShareCoreObject::IsFrameSyncActive() const
 {
-	return IsSessionActive() && IsActive() && !FrameConnections.IsEmpty();
+	return IsSessionActive() && IsActive() && !IsEmptyFrameConnections();
 }
 
 bool FTextureShareCoreObject::IsFrameSyncActive_RenderThread() const
 {
-	return IsSessionActive() && IsActive_RenderThread() && !FrameConnections.IsEmpty();
+	return IsSessionActive() && IsActive_RenderThread() && !IsEmptyFrameConnections();
 }
 
 /////////////////////////////////////////////////////////////////////////////////////////////////////////////////////////////////
@@ -78,7 +94,7 @@ bool FTextureShareCoreObject::IsBeginFrameSyncActive() const
 
 bool FTextureShareCoreObject::IsBeginFrameSyncActive_RenderThread() const
 {
-	if (FrameConnections.IsEmpty())
+	if (IsEmptyFrameConnections())
 	{
 		UE_TS_LOG(LogTextureShareCoreProxyObjectSync, Log, TEXT("%s:IsBeginFrameSyncActive_RenderThread() - canceled: no connections"), *GetName());
 
@@ -100,9 +116,3 @@ bool FTextureShareCoreObject::IsBeginFrameSyncActive_RenderThread() const
 
 	return true;
 }
-
-//////////////////////////////////////////////////////////////////////////////////////////////
-const TArraySerializable<FTextureShareCoreObjectDesc>& FTextureShareCoreObject::GetConnectedInterprocessObjects() const
-{
-	return FrameConnections;
-}
diff --git a/Engine/Plugins/VirtualProduction/TextureShare/Source/TextureShareCore/Private/Object/TextureShareCoreObject_Sync.cpp b/Engine/Plugins/VirtualProduction/TextureShare/Source/TextureShareCore/Private/Object/TextureShareCoreObject_Sync.cpp
--- a/Engine/Plugins/VirtualProduction/TextureShare/Source/TextureShareCore/Private/Object/TextureShareCoreObject_Sync.cpp
+++ b/Engine/Plugins/VirtualProduction/TextureShare/Source/TextureShareCore/Private/Object/TextureShareCoreObject_Sync.cpp
@@ -34,16 +34,14 @@ bool FTextureShareCoreObject::FindSkippedSyncStep_RenderThread(const ETextureSha
 bool FTextureShareCoreObject::FindSkippedSyncStepImpl(const ETextureShareSyncStep InSyncStep, bool bIsProxyFrame, ETextureShareSyncStep& OutSkippedSyncStep)
 {
 	// Update sync steps on request
-	if (!SyncSettings.FrameSyncSettings.Steps.Contains(InSyncStep))
-	{
-		AddNewSyncStep(InSyncStep);
-	}
+	AddNewSyncStep(InSyncStep);
 
-	if(const int32 NextStepIndex = SyncSettings.FrameSyncSettings.Steps.Find(CurrentSyncStep) + 1)
+	const TArraySerializable<ETextureShareSyncStep> SyncSteps = GetSyncSettings().FrameSyncSettings.Steps;
+	if(const int32 NextStepIndex = SyncSteps.Find(CurrentSyncStep) + 1)
 	{
 		const ETextureShareSyncStep FrameEndSyncStep = bIsProxyFrame ? ETextureShareSyncStep::FrameProxyEnd : ETextureShareSyncStep::FrameEnd;
 
-		OutSkippedSyncStep = SyncSettings.FrameSyncSettings.Steps[NextStepIndex];
+		OutSkippedSyncStep = SyncSteps[NextStepIndex];
 		if (OutSkippedSyncStep < InSyncStep && OutSkippedSyncStep < FrameEndSyncStep)
 		{
 			return true;
@@ -390,7 +388,7 @@ void FTextureShareCoreObject::SetFrameSyncState(const ETextureShareCoreInterproc
 //////////////////////////////////////////////////////////////////////////////////////////////
 void FTextureShareCoreObject::UpdateLastAccessTime() const
 {
-	if (IsSessionActive() && IsActive() && Owner.LockInterprocessMemory(SyncSettings.TimeoutSettings.MemoryMutexTimeout))
+	if (IsSessionActive() && IsActive() && Owner.LockInterprocessMemory(GetTimeOutSettings().MemoryMutexTimeout))
 	{
 		if (FTextureShareCoreInterprocessMemory* InterprocessMemory = Owner.GetInterprocessMemory())
 		{
@@ -409,10 +407,93 @@ void FTextureShareCoreObject::ReleaseSyncData()
 {
 	UE_TS_LOG(LogTextureShareCoreObject, Log, TEXT("%s:ReleaseSyncData()"), *GetName());
 
-	FrameConnections.Empty();
+	ResetFrameConnections();
 
 	FrameSyncState = ETextureShareCoreInterprocessObjectFrameSyncState::Undefined;
 	CurrentSyncStep = ETextureShareSyncStep::Undefined;
 
 	CachedNotificationEvents.Empty();
 }
+
+//////////////////////////////////////////////////////////////////////////////////////////////
+TArraySerializable<FTextureShareCoreObjectDesc> FTextureShareCoreObject::GetConnectedInterprocessObjects() const
+{
+	return GetFrameConnections();
+}
+
+void FTextureShareCoreObject::UpdateFrameConnections(FTextureShareCoreInterprocessMemory& InterprocessMemory)
+{
+	if (LockThreadMutex(ETextureShareThreadMutex::InternalFrameConnectionsMT))
+	{
+		InterprocessMemory.UpdateFrameConnections(FrameConnectionsMT);
+
+		UnlockThreadMutex(ETextureShareThreadMutex::InternalFrameConnectionsMT);
+	}
+}
+
+int32 FTextureShareCoreObject::FindFrameConnections(FTextureShareCoreInterprocessMemory& InterprocessMemory, FTextureShareCoreInterprocessObject& LocalObject)
+{
+	int32 ProcessNum = 0;
+	if (LockThreadMutex(ETextureShareThreadMutex::InternalFrameConnectionsMT))
+	{
+		ProcessNum = InterprocessMemory.FindConnectableObjects(FrameConnectionsMT, LocalObject);
+
+		UnlockThreadMutex(ETextureShareThreadMutex::InternalFrameConnectionsMT);
+	}
+
+	return ProcessNum;
+}
+
+void FTextureShareCoreObject::SetFrameConnections(const TArraySerializable<FTextureShareCoreObjectDesc>& InNewFrameConnections)
+{
+	if (LockThreadMutex(ETextureShareThreadMutex::InternalFrameConnectionsMT))
+	{
+		FrameConnectionsMT = InNewFrameConnections;
+
+		UnlockThreadMutex(ETextureShareThreadMutex::InternalFrameConnectionsMT);
+	}
+}
+
+void FTextureShareCoreObject::ResetFrameConnections()
+{
+	if (!IsSessionActive())
+	{
+		FrameConnectionsMT.Reset();
+	}
+	else if (LockThreadMutex(ETextureShareThreadMutex::InternalFrameConnectionsMT))
+	{
+		FrameConnectionsMT.Reset();
+
+		UnlockThreadMutex(ETextureShareThreadMutex::InternalFrameConnectionsMT);
+	}
+}
+
+bool FTextureShareCoreObject::IsEmptyFrameConnections() const
+{
+	bool bIsEmptyFrameConnections = true;
+	if (LockThreadMutex(ETextureShareThreadMutex::InternalFrameConnectionsMT))
+	{
+		bIsEmptyFrameConnections = FrameConnectionsMT.IsEmpty();
+
+		UnlockThreadMutex(ETextureShareThreadMutex::InternalFrameConnectionsMT);
+	}
+	
+	return bIsEmptyFrameConnections;
+}
+
+const TArraySerializable<FTextureShareCoreObjectDesc> FTextureShareCoreObject::GetFrameConnections() const
+{
+	TArraySerializable<FTextureShareCoreObjectDesc> OutFrameConnections;
+	if (!IsSessionActive())
+	{
+		OutFrameConnections = FrameConnectionsMT;
+	}
+	else  if (LockThreadMutex(ETextureShareThreadMutex::InternalFrameConnectionsMT))
+	{
+		OutFrameConnections = FrameConnectionsMT;
+
+		UnlockThreadMutex(ETextureShareThreadMutex::InternalFrameConnectionsMT);
+	}
+
+	return OutFrameConnections;
+}
\ No newline at end of file
diff --git a/Engine/Plugins/VirtualProduction/TextureShare/Source/TextureShareCore/Private/Object/TextureShareCoreObject_SyncBarrier.cpp b/Engine/Plugins/VirtualProduction/TextureShare/Source/TextureShareCore/Private/Object/TextureShareCoreObject_SyncBarrier.cpp
--- a/Engine/Plugins/VirtualProduction/TextureShare/Source/TextureShareCore/Private/Object/TextureShareCoreObject_SyncBarrier.cpp
+++ b/Engine/Plugins/VirtualProduction/TextureShare/Source/TextureShareCore/Private/Object/TextureShareCoreObject_SyncBarrier.cpp
@@ -21,7 +21,7 @@ using namespace UE::TextureShareCore;
 //////////////////////////////////////////////////////////////////////////////////////////////
 bool FTextureShareCoreObject::BeginSyncBarrier(FTextureShareCoreInterprocessMemory& InterprocessMemory, FTextureShareCoreInterprocessObject& LocalObject, const ETextureShareSyncStep InSyncStep, const ETextureShareSyncPass InSyncPass)
 {
-	UE_TS_BARRIER_LOG(LogTextureShareCoreObjectSync, Log, TEXT("%s:BeginSyncBarrier(%s[%s]) %s In << %s >>"), *GetName(), GetTEXT(InSyncStep), GetTEXT(InSyncPass), *ToString(LocalObject), *ToString(FrameConnections));
+	UE_TS_BARRIER_LOG(LogTextureShareCoreObjectSync, Log, TEXT("%s:BeginSyncBarrier(%s[%s]) %s In << %s >>"), *GetName(), GetTEXT(InSyncStep), GetTEXT(InSyncPass), *ToString(LocalObject), *ToString(GetFrameConnections()));
 
 	if (LocalObject.Sync.BeginSyncBarrier(InSyncStep, InSyncPass, FindNextSyncStep(InSyncStep)))
 	{
@@ -33,7 +33,7 @@ bool FTextureShareCoreObject::BeginSyncBarrier(FTextureShareCoreInterprocessMemo
 
 bool FTextureShareCoreObject::AcceptSyncBarrier(FTextureShareCoreInterprocessMemory& InterprocessMemory, FTextureShareCoreInterprocessObject& LocalObject, const ETextureShareSyncStep InSyncStep, const ETextureShareSyncPass InSyncPass)
 {
-	UE_TS_BARRIER_LOG(LogTextureShareCoreObjectSync, Log, TEXT("%s:AcceptSyncBarrier(%s[%s]) %s In << %s >>"), *GetName(), GetTEXT(InSyncStep), GetTEXT(InSyncPass), *ToString(LocalObject), *ToString(FrameConnections));
+	UE_TS_BARRIER_LOG(LogTextureShareCoreObjectSync, Log, TEXT("%s:AcceptSyncBarrier(%s[%s]) %s In << %s >>"), *GetName(), GetTEXT(InSyncStep), GetTEXT(InSyncPass), *ToString(LocalObject), *ToString(GetFrameConnections()));
 
 	if (LocalObject.Sync.AcceptSyncBarrier(InSyncStep, InSyncPass))
 	{
@@ -49,12 +49,13 @@ bool FTextureShareCoreObject::PrepareSyncBarrierPass(const ETextureShareSyncStep
 	bool bResult = false;
 
 	// Return objects who support sync pass at this frame
-	if (IsFrameSyncActive() && !FrameConnections.IsEmpty() && Owner.LockInterprocessMemory(SyncSettings.TimeoutSettings.MemoryMutexTimeout))
+	if (IsFrameSyncActive() && !IsEmptyFrameConnections() && Owner.LockInterprocessMemory(GetTimeOutSettings().MemoryMutexTimeout))
 	{
 		if (FTextureShareCoreInterprocessMemory* InterprocessMemory = Owner.GetInterprocessMemory())
 		{
-			InterprocessMemory->UpdateFrameConnections(FrameConnections);
-			bResult = InterprocessMemory->IsUsedSyncBarrierPass(FrameConnections, InSyncStep);
+			UpdateFrameConnections(*InterprocessMemory);
+
+			bResult = InterprocessMemory->IsUsedSyncBarrierPass(GetFrameConnections(), InSyncStep);
 		}
 
 		Owner.UnlockInterprocessMemory();
@@ -68,12 +69,13 @@ bool FTextureShareCoreObject::PrepareSyncBarrierPass_RenderThread(const ETexture
 	bool bResult = false;
 
 	// Return objects who support sync pass at this frame
-	if (IsFrameSyncActive_RenderThread() && !FrameConnections.IsEmpty() && Owner.LockInterprocessMemory(SyncSettings.TimeoutSettings.MemoryMutexTimeout))
+	if (IsFrameSyncActive_RenderThread() && !IsEmptyFrameConnections() && Owner.LockInterprocessMemory(GetTimeOutSettings().MemoryMutexTimeout))
 	{
 		if (FTextureShareCoreInterprocessMemory* InterprocessMemory = Owner.GetInterprocessMemory())
 		{
-			InterprocessMemory->UpdateFrameConnections(FrameConnections);
-			bResult = InterprocessMemory->IsUsedSyncBarrierPass(FrameConnections, InSyncStep);
+			UpdateFrameConnections(*InterprocessMemory);
+
+			bResult = InterprocessMemory->IsUsedSyncBarrierPass(GetFrameConnections(), InSyncStep);
 		}
 
 		Owner.UnlockInterprocessMemory();
@@ -91,7 +93,7 @@ bool FTextureShareCoreObject::TryEnterSyncBarrier(const ETextureShareSyncStep In
 	case ETextureShareCoreInterprocessObjectFrameSyncState::FrameBegin:
 		if ((InSyncStep > ETextureShareSyncStep::FrameBegin) && (InSyncStep < ETextureShareSyncStep::FrameEnd))
 		{
-			if (GetSyncSetting().FrameSyncSettings.Steps.Contains(InSyncStep))
+			if (GetSyncSettings().FrameSyncSettings.Steps.Contains(InSyncStep))
 			{
 				return true;
 			}
@@ -105,7 +107,7 @@ bool FTextureShareCoreObject::TryEnterSyncBarrier(const ETextureShareSyncStep In
 	case ETextureShareCoreInterprocessObjectFrameSyncState::FrameProxyBegin:
 		if ((InSyncStep > ETextureShareSyncStep::FrameProxyBegin) && (InSyncStep < ETextureShareSyncStep::FrameProxyEnd))
 		{
-			if (GetSyncSetting().FrameSyncSettings.Steps.Contains(InSyncStep))
+			if (GetSyncSettings().FrameSyncSettings.Steps.Contains(InSyncStep))
 			{
 				return true;
 			}
@@ -130,7 +132,7 @@ bool FTextureShareCoreObject::TryFrameProcessesBarrier(FTextureShareCoreInterpro
 	// Mark last acces time for current process (this value used by other processes to detect die)
 	LocalObject.Sync.UpdateLastAccessTime();
 
-	if (FrameConnections.IsEmpty())
+	if (IsEmptyFrameConnections())
 	{
 		UE_TS_BARRIER_LOG(LogTextureShareCoreObjectSync, Error, TEXT("%s:TryFrameProcessesBarrier(%s) Failed"), *GetName(), *ToString(LocalObject));
 
@@ -138,12 +140,12 @@ bool FTextureShareCoreObject::TryFrameProcessesBarrier(FTextureShareCoreInterpro
 	}
 
 #if TEXTURESHARECORE_BARRIER_DEBUGLOG
-	InterprocessMemory.UpdateFrameConnections(FrameConnections);
-	UE_TS_BARRIER_LOG(LogTextureShareCoreObjectSync, Log, TEXT("%s:TryFrameProcessesBarrier(%s) %s"), *GetName(), *ToString(LocalObject), *ToString(FrameConnections));
+	UpdateFrameConnections(InterprocessMemory);
+	UE_TS_BARRIER_LOG(LogTextureShareCoreObjectSync, Log, TEXT("%s:TryFrameProcessesBarrier(%s) %s"), *GetName(), *ToString(LocalObject), *ToString(GetFrameConnections()));
 #endif
 
 
-	const ETextureShareInterprocessObjectSyncBarrierState BarrierState = InterprocessMemory.GetBarrierState(LocalObject, FrameConnections);
+	const ETextureShareInterprocessObjectSyncBarrierState BarrierState = InterprocessMemory.GetBarrierState(LocalObject, GetFrameConnections());
 	switch (BarrierState)
 	{
 	case ETextureShareInterprocessObjectSyncBarrierState::Accept:
@@ -189,7 +191,7 @@ bool FTextureShareCoreObject::SyncBarrierPass(const ETextureShareSyncStep InSync
 	UE_TS_BARRIER_LOG(LogTextureShareCoreObjectSync, Log, TEXT("%s:SyncBarrierPass(%s, %s)"), *GetName(), GetTEXT(InSyncStep), GetTEXT(InSyncPass));
 
 	bool bBarrierResult = false;
-	if (IsFrameSyncActive() && Owner.LockInterprocessMemory(SyncSettings.TimeoutSettings.MemoryMutexTimeout))
+	if (IsFrameSyncActive() && Owner.LockInterprocessMemory(GetTimeOutSettings().MemoryMutexTimeout))
 	{
 		if (FTextureShareCoreInterprocessMemory* InterprocessMemory = Owner.GetInterprocessMemory())
 		{
@@ -220,7 +222,7 @@ bool FTextureShareCoreObject::SyncBarrierPass_RenderThread(const ETextureShareSy
 	UE_TS_BARRIER_LOG(LogTextureShareCoreObjectSync, Log, TEXT("%s:SyncBarrierPass_RenderThread(%s, %s)"), *GetName(), GetTEXT(InSyncStep), GetTEXT(InSyncPass));
 
 	bool bBarrierResult = false;
-	if (IsFrameSyncActive_RenderThread() && Owner.LockInterprocessMemory(SyncSettings.TimeoutSettings.MemoryMutexTimeout))
+	if (IsFrameSyncActive_RenderThread() && Owner.LockInterprocessMemory(GetTimeOutSettings().MemoryMutexTimeout))
 	{
 		if (FTextureShareCoreInterprocessMemory* InterprocessMemory = Owner.GetInterprocessMemory())
 		{
@@ -252,7 +254,7 @@ bool FTextureShareCoreObject::SyncBarrierPassImpl(FTextureShareCoreInterprocessM
 	BeginSyncBarrier(InterprocessMemory, LocalObject, InSyncStep, InSyncPass);
 
 	// Begin processs barrier sync
-	FTextureShareCoreObjectTimeout FrameSyncTimer(SyncSettings.TimeoutSettings.FrameSyncTimeOut, SyncSettings.TimeoutSettings.FrameSyncTimeOutSplit);
+	FTextureShareCoreObjectTimeout FrameSyncTimer(GetTimeOutSettings().FrameSyncTimeOut, GetTimeOutSettings().FrameSyncTimeOutSplit);
 
 	// repeat that barrier, until the connected processes is defined
 	while (TryFrameProcessesBarrier(InterprocessMemory, LocalObject, InSyncStep, InSyncPass))
@@ -274,12 +276,12 @@ bool FTextureShareCoreObject::SyncBarrierPassImpl(FTextureShareCoreInterprocessM
 		}
 	}
 
-	return LocalObject.IsEnabled() && FrameConnections.Num() > 0;
+	return LocalObject.IsEnabled() && !IsEmptyFrameConnections();
 }
 
 ETextureShareSyncStep FTextureShareCoreObject::FindNextSyncStep(const ETextureShareSyncStep InSyncStep) const
 {
-	const TArray<ETextureShareSyncStep>& SyncSteps = SyncSettings.FrameSyncSettings.Steps;
+	const TArray<ETextureShareSyncStep> SyncSteps = GetSyncSettings().FrameSyncSettings.Steps;
 	if (SyncSteps.IsEmpty() == false)
 	{
 		int32 Index = (InSyncStep == ETextureShareSyncStep::InterprocessConnection) ? 0 : SyncSteps.Find(InSyncStep);
@@ -297,14 +299,15 @@ ETextureShareSyncStep FTextureShareCoreObject::FindNextSyncStep(const ETextureSh
 				case ETextureShareSyncStep::Undefined:
 					// Skip internal values
 					break;
+
 				default:
-					if (FrameConnections.IsEmpty())
+					if (IsEmptyFrameConnections())
 					{
 						return NextStep;
 					}
 
 					// Use next step only if any process support
-					for (const FTextureShareCoreObjectDesc& ObjIt : FrameConnections)
+					for (const FTextureShareCoreObjectDesc& ObjIt : GetFrameConnections())
 					{
 						if (ObjIt.Sync.IsStepEnabled(NextStep))
 						{
diff --git a/Engine/Plugins/VirtualProduction/TextureShare/Source/TextureShareCore/Private/Object/TextureShareCoreObject_ThreadMutex.cpp b/Engine/Plugins/VirtualProduction/TextureShare/Source/TextureShareCore/Private/Object/TextureShareCoreObject_ThreadMutex.cpp
--- a/Engine/Plugins/VirtualProduction/TextureShare/Source/TextureShareCore/Private/Object/TextureShareCoreObject_ThreadMutex.cpp
+++ b/Engine/Plugins/VirtualProduction/TextureShare/Source/TextureShareCore/Private/Object/TextureShareCoreObject_ThreadMutex.cpp
@@ -13,20 +13,54 @@
 #include "Module/TextureShareCoreLog.h"
 
 #include "Misc/ScopeLock.h"
+#include "HAL/IConsoleManager.h"
 
 using namespace UE::TextureShareCore;
 
+/** Use lock for internal data that can be used by multiple threads at the same time. */
+int32 GTextureShareMultithreadDataLock = 1;
+static FAutoConsoleVariableRef CVarTextureShareMultithreadDataLock(
+	TEXT("TextureShare.MultithreadDataLock"),
+	GTextureShareMultithreadDataLock,
+	TEXT("Lock multithread data by the mutex. (Default = 1)"),
+	ECVF_Default
+);
+
+/** Helper functions. */
+namespace UE::TextureShareCore::Helpers
+{
+	/** true if this thread mutex enabled. */
+	static inline bool IsThreadMutexEnabled(const ETextureShareThreadMutex InThreadMutex)
+	{
+		switch (InThreadMutex)
+		{
+		case ETextureShareThreadMutex::InternalObjectDescMT:
+		case ETextureShareThreadMutex::InternalFrameConnectionsMT:
+		case ETextureShareThreadMutex::InternalSyncSettingsMT:
+			return GTextureShareMultithreadDataLock != 0;
+			break;
+		}
+
+		return true;
+	}
+};
+
 //////////////////////////////////////////////////////////////////////////////////////////////
 // FTextureShareCoreObject
 //////////////////////////////////////////////////////////////////////////////////////////////
-bool FTextureShareCoreObject::LockThreadMutex(const ETextureShareThreadMutex InThreadMutex, bool bForceLockNoWait)
+bool FTextureShareCoreObject::LockThreadMutex(const ETextureShareThreadMutex InThreadMutex, bool bForceLockNoWait) const
 {
+	if (!UE::TextureShareCore::Helpers::IsThreadMutexEnabled(InThreadMutex))
+	{
+		return true;
+	}
+
 	TSharedPtr<FTextureShareCoreInterprocessMutex, ESPMode::ThreadSafe> ThreadMutex = GetThreadMutex(InThreadMutex);
 	if (!ThreadMutex.IsValid())
 	{
 		UE_TS_LOG(LogTextureShareCoreObjectSync, Error, TEXT("%s:LockThreadMutex(%s) Mutex object not exist"), *GetName(), GetTEXT(InThreadMutex));
 
-		return false;
+		return true;
 	}
 
 	if (bForceLockNoWait)
@@ -36,11 +70,17 @@ bool FTextureShareCoreObject::LockThreadMutex(const ETextureShareThreadMutex InT
 		return ThreadMutex->LockMutex(0);
 	}
 
-	UE_TS_LOG(LogTextureShareCoreObjectSync, VeryVerbose, TEXT("%s:LockThreadMutex(%s) try"), *GetName(), GetTEXT(InThreadMutex));
+	if (InThreadMutex < ETextureShareThreadMutex::MAX_LOG)
+	{
+		UE_TS_LOG(LogTextureShareCoreObjectSync, VeryVerbose, TEXT("%s:LockThreadMutex(%s) try"), *GetName(), GetTEXT(InThreadMutex));
+	}
 
-	if (ThreadMutex->LockMutex(SyncSettings.TimeoutSettings.ThreadMutexTimeout))
+	if (ThreadMutex->LockMutex(GetTimeOutSettings().ThreadMutexTimeout))
 	{
-		UE_TS_LOG(LogTextureShareCoreObjectSync, Log, TEXT("%s:LockThreadMutex(%s)"), *GetName(), GetTEXT(InThreadMutex));
+		if (InThreadMutex < ETextureShareThreadMutex::MAX_LOG)
+		{
+			UE_TS_LOG(LogTextureShareCoreObjectSync, Log, TEXT("%s:LockThreadMutex(%s)"), *GetName(), GetTEXT(InThreadMutex));
+		}
 
 		return true;
 	}
@@ -54,19 +94,29 @@ bool FTextureShareCoreObject::LockThreadMutex(const ETextureShareThreadMutex InT
 	return false;
 }
 
-bool FTextureShareCoreObject::UnlockThreadMutex(const ETextureShareThreadMutex InThreadMutex)
+bool FTextureShareCoreObject::UnlockThreadMutex(const ETextureShareThreadMutex InThreadMutex) const
 {
-	TSharedPtr<FTextureShareCoreInterprocessMutex, ESPMode::ThreadSafe> ThreadMutex = GetThreadMutex(InThreadMutex);
-	if (ThreadMutex.IsValid())
+	if (!UE::TextureShareCore::Helpers::IsThreadMutexEnabled(InThreadMutex))
 	{
-		UE_TS_LOG(LogTextureShareCoreObjectSync, Log, TEXT("%s:UnlockThreadMutex(%s)"), *GetName(), GetTEXT(InThreadMutex));
+		return true;
+	}
 
-		ThreadMutex->UnlockMutex();
+	TSharedPtr<FTextureShareCoreInterprocessMutex, ESPMode::ThreadSafe> ThreadMutex = GetThreadMutex(InThreadMutex);
+	if (!ThreadMutex.IsValid())
+	{
+		UE_TS_LOG(LogTextureShareCoreObjectSync, Error, TEXT("%s:UnlockThreadMutex(%s) Mutex object not exist"), *GetName(), GetTEXT(InThreadMutex));
 
 		return true;
 	}
 
-	return false;
+	if (InThreadMutex < ETextureShareThreadMutex::MAX_LOG)
+	{
+		UE_TS_LOG(LogTextureShareCoreObjectSync, Log, TEXT("%s:UnlockThreadMutex(%s)"), *GetName(), GetTEXT(InThreadMutex));
+	}
+
+	ThreadMutex->UnlockMutex();
+
+	return true;
 }
 
 //////////////////////////////////////////////////////////////////////////////////////////////
@@ -82,7 +132,7 @@ void FTextureShareCoreObject::InitializeThreadMutexes()
 	}
 }
 
-TSharedPtr<FTextureShareCoreInterprocessMutex, ESPMode::ThreadSafe> FTextureShareCoreObject::GetThreadMutex(const ETextureShareThreadMutex InThreadMutex)
+TSharedPtr<FTextureShareCoreInterprocessMutex, ESPMode::ThreadSafe> FTextureShareCoreObject::GetThreadMutex(const ETextureShareThreadMutex InThreadMutex) const
 {
 	const uint8 ThreadMutexIndex = (uint8)InThreadMutex;
 	if (ThreadMutexes.IsValidIndex(ThreadMutexIndex))
@@ -98,13 +148,25 @@ TSharedPtr<FTextureShareCoreInterprocessMutex, ESPMode::ThreadSafe> FTextureShar
 
 void FTextureShareCoreObject::ResetThreadMutex(const ETextureShareThreadMutex InThreadMutex)
 {
+	if (!UE::TextureShareCore::Helpers::IsThreadMutexEnabled(InThreadMutex))
+	{
+		return;
+	}
+
 	TSharedPtr<FTextureShareCoreInterprocessMutex, ESPMode::ThreadSafe> ThreadMutex = GetThreadMutex(InThreadMutex);
-	if (ThreadMutex.IsValid())
+	if (!ThreadMutex.IsValid())
 	{
-		UE_TS_LOG(LogTextureShareCoreObjectSync, Error, TEXT("%s:ResetThreadMutex(%s)"), *GetName(), GetTEXT(InThreadMutex));
+		UE_TS_LOG(LogTextureShareCoreObjectSync, Error, TEXT("%s:ResetThreadMutex(%s) Mutex object not exist"), *GetName(), GetTEXT(InThreadMutex));
 
-		ThreadMutex->TryUnlockMutex();
+		return;
 	}
+
+	if (InThreadMutex < ETextureShareThreadMutex::MAX_LOG)
+	{
+		UE_TS_LOG(LogTextureShareCoreObjectSync, Error, TEXT("%s:ResetThreadMutex(%s)"), *GetName(), GetTEXT(InThreadMutex));
+	}
+
+	ThreadMutex->TryUnlockMutex();
 }
 
 void FTextureShareCoreObject::ReleaseThreadMutexes()
diff --git a/Engine/Source/Programs/VirtualProduction/TextureShare/Private/TextureShareSDKObject.cpp b/Engine/Source/Programs/VirtualProduction/TextureShare/Private/TextureShareSDKObject.cpp
--- a/Engine/Source/Programs/VirtualProduction/TextureShare/Private/TextureShareSDKObject.cpp
+++ b/Engine/Source/Programs/VirtualProduction/TextureShare/Private/TextureShareSDKObject.cpp
@@ -116,13 +116,13 @@ bool FTextureShareSDKObject::SetDeviceType(const ETextureShareDeviceType InDevic
 bool FTextureShareSDKObject::SetSyncSetting(const TDataInput<FTextureShareCoreSyncSettings>& InSyncSetting)
 {
 	check(CoreObject.IsValid());
-	return CoreObject->SetSyncSetting(*TDataVariable<FTextureShareCoreSyncSettings>(InSyncSetting));
+	return CoreObject->SetSyncSettings(*TDataVariable<FTextureShareCoreSyncSettings>(InSyncSetting));
 }
 
 void FTextureShareSDKObject::GetSyncSetting(TDataOutput<FTextureShareCoreSyncSettings>& OutSyncSettings) const
 {
 	check(CoreObject.IsValid());
-	OutSyncSettings = CoreObject->GetSyncSetting();
+	OutSyncSettings = CoreObject->GetSyncSettings();
 }
 
 void FTextureShareSDKObject::GetFrameSyncSettings(const ETextureShareFrameSyncTemplate InType, TDataOutput<FTextureShareCoreFrameSyncSettings>& OutFrameSyncSettings) const
@@ -260,9 +260,10 @@ ID3D11Texture2D* FTextureShareSDKObject::OpenSharedResourceD3D11(ID3D11Device* p
 	if (CoreObject.IsValid() && D3D11Resources.IsValid())
 	{
 		FTextureShareCoreResourceDesc ResourceDesc(*TDataVariable<FTextureShareCoreResourceDesc>(InResourceDesc));
-		if (const FTextureShareCoreResourceHandle* ResourceHandle = ImplFindResource(ResourceDesc))
+		FTextureShareCoreResourceHandle ResourceHandle;
+		if (ImplFindResource(ResourceDesc, ResourceHandle))
 		{
-			return D3D11Resources->OpenSharedResource(CoreObject->GetObjectDesc(), pD3D11Device, *ResourceHandle);
+			return D3D11Resources->OpenSharedResource(CoreObject->GetObjectDesc(), pD3D11Device, ResourceHandle);
 		}
 	}
 
@@ -275,9 +276,10 @@ ID3D12Resource* FTextureShareSDKObject::OpenSharedResourceD3D12(ID3D12Device* pD
 	if (CoreObject.IsValid() && D3D12Resources.IsValid())
 	{
 		FTextureShareCoreResourceDesc ResourceDesc(*TDataVariable<FTextureShareCoreResourceDesc>(InResourceDesc));
-		if (const FTextureShareCoreResourceHandle* ResourceHandle = ImplFindResource(ResourceDesc))
+		FTextureShareCoreResourceHandle ResourceHandle;
+		if (ImplFindResource(ResourceDesc, ResourceHandle))
 		{
-			return D3D12Resources->OpenSharedResource(CoreObject->GetObjectDesc(), pD3D12Device, *ResourceHandle);
+			return D3D12Resources->OpenSharedResource(CoreObject->GetObjectDesc(), pD3D12Device, ResourceHandle);
 		}
 	}
 
@@ -290,10 +292,11 @@ bool FTextureShareSDKObject::OpenSharedResourceVulkan(const TDataInput<FTextureS
 	if (CoreObject.IsValid() && VulkanResources.IsValid())
 	{
 		FTextureShareCoreResourceDesc ResourceDesc(*TDataVariable<FTextureShareCoreResourceDesc>(InResourceDesc));
-		if (const FTextureShareCoreResourceHandle* ResourceHandle = ImplFindResource(ResourceDesc))
+		FTextureShareCoreResourceHandle ResourceHandle;
+		if (ImplFindResource(ResourceDesc, ResourceHandle))
 		{
 			FTextureShareDeviceVulkanContext DeviceVulkanContext(*TDataVariable<FTextureShareDeviceVulkanContext>(InDeviceVulkanContext));
-			FTextureShareDeviceVulkanResource VulkanResource = VulkanResources->OpenSharedResource(CoreObject->GetObjectDesc(), DeviceVulkanContext, *ResourceHandle);
+			FTextureShareDeviceVulkanResource VulkanResource = VulkanResources->OpenSharedResource(CoreObject->GetObjectDesc(), DeviceVulkanContext, ResourceHandle);
 			if (VulkanResource.IsValid())
 			{
 				OutVulkanResource = VulkanResource;
@@ -307,20 +310,21 @@ bool FTextureShareSDKObject::OpenSharedResourceVulkan(const TDataInput<FTextureS
 }
 
 //////////////////////////////////////////////////////////
-const FTextureShareCoreResourceHandle* FTextureShareSDKObject::ImplFindResource(const FTextureShareCoreResourceDesc& InResourceDesc) const
+bool FTextureShareSDKObject::ImplFindResource(const FTextureShareCoreResourceDesc& InResourceDesc, FTextureShareCoreResourceHandle& OutResourceHandle) const
 {
 	if (CoreObject.IsValid())
 	{
-		const TArraySerializable<FTextureShareCoreObjectProxyData>& ReceivedObjectsProxyData = CoreObject->GetReceivedProxyData_RenderThread();
+		const TArraySerializable<FTextureShareCoreObjectProxyData> ReceivedObjectsProxyData = CoreObject->GetReceivedProxyData_RenderThread();
 		for (const FTextureShareCoreObjectProxyData& ObjectProxyDataIt : ReceivedObjectsProxyData)
 		{
 			if (const FTextureShareCoreResourceHandle* ExistResource = ObjectProxyDataIt.ProxyData.ResourceHandles.FindByEqualsFunc(InResourceDesc))
 			{
 				// Get resource from any process
-				return ExistResource;
+				OutResourceHandle = *ExistResource;
+				return true;
 			}
 		}
 	}
 
-	return nullptr;
+	return false;
 }
diff --git a/Engine/Source/Programs/VirtualProduction/TextureShare/Private/TextureShareSDKObject.h b/Engine/Source/Programs/VirtualProduction/TextureShare/Private/TextureShareSDKObject.h
--- a/Engine/Source/Programs/VirtualProduction/TextureShare/Private/TextureShareSDKObject.h
+++ b/Engine/Source/Programs/VirtualProduction/TextureShare/Private/TextureShareSDKObject.h
@@ -80,7 +80,7 @@ public:
 	virtual  bool OpenSharedResourceVulkan(const TDataInput<FTextureShareDeviceVulkanContext>& InDeviceVulkanContext, const TDataInput<FTextureShareCoreResourceDesc>& InResourceDesc, TDataOutput<FTextureShareDeviceVulkanResource>& OutVulkanResource) override;
 
 private:
-	const struct FTextureShareCoreResourceHandle* ImplFindResource(const struct FTextureShareCoreResourceDesc& InResourceDesc) const;
+	bool ImplFindResource(const struct FTextureShareCoreResourceDesc& InResourceDesc, struct FTextureShareCoreResourceHandle& OutResourceHandle) const;
 
 private:
 	TSharedPtr<ITextureShareCoreObject, ESPMode::ThreadSafe> const CoreObject;
diff --git a/Engine/Source/Programs/VirtualProduction/TextureShare/TextureShareSDK.Build.cs b/Engine/Source/Programs/VirtualProduction/TextureShare/TextureShareSDK.Build.cs
--- a/Engine/Source/Programs/VirtualProduction/TextureShare/TextureShareSDK.Build.cs
+++ b/Engine/Source/Programs/VirtualProduction/TextureShare/TextureShareSDK.Build.cs
@@ -22,6 +22,7 @@ namespace UnrealBuildTool.Rules
 				new string[]
 				{
 					"Core",
+					"ImageCore",
 					"TextureShareCore"
 				}
 			);
