[VirtualProduction] [nDisplay] [TextureShare] - fixed Automation test errors for mutex objects.

- Prevent the Unlock() function from being called multiple times for the FPlatformProcess::FSemaphore object. Added TryLock(0) in the UnlockThread().

- removed macro LOG_SCOPE_VERBOSITY_OVERRIDE. It can be used only in the game thread.

- The macros LOG_SCOPE_VERBOSITY_OVERRIDE were removed. It requires the CL 40816100.

- Removed unsued line in the FTextureShareAPI::RemoveObject() that leads to the crash when ENQUEUE_RENDER_COMMAND() is completed before the function exit.

- Code cleanup

diff --git a/Engine/Plugins/VirtualProduction/TextureShare/Source/TextureShare/Private/Module/TextureShareAPI.cpp b/Engine/Plugins/VirtualProduction/TextureShare/Source/TextureShare/Private/Module/TextureShareAPI.cpp
--- a/Engine/Plugins/VirtualProduction/TextureShare/Source/TextureShare/Private/Module/TextureShareAPI.cpp
+++ b/Engine/Plugins/VirtualProduction/TextureShare/Source/TextureShare/Private/Module/TextureShareAPI.cpp
@@ -123,7 +123,6 @@ bool FTextureShareAPI::RemoveObject(const FString& ShareName)
 		return false;
 	}
 
-	TSharedPtr<FTextureShareObject, ESPMode::ThreadSafe> Object = Objects[ShareNameLwr];
 	Objects[ShareNameLwr].Reset();
 	Objects.Remove(ShareNameLwr);
 
diff --git a/Engine/Plugins/VirtualProduction/TextureShare/Source/TextureShareCore/Private/Core/TextureShareCore.cpp b/Engine/Plugins/VirtualProduction/TextureShare/Source/TextureShareCore/Private/Core/TextureShareCore.cpp
--- a/Engine/Plugins/VirtualProduction/TextureShare/Source/TextureShareCore/Private/Core/TextureShareCore.cpp
+++ b/Engine/Plugins/VirtualProduction/TextureShare/Source/TextureShareCore/Private/Core/TextureShareCore.cpp
@@ -190,8 +190,8 @@ bool FTextureShareCore::InitializeInterprocessMemory()
 	// An open memory mutex that is used to synchronize all shared memory operations.
 	if (!InterprocessMemoryMutex.IsValid())
 	{
-		InterprocessMemoryMutex = MakeShared<FTextureShareCoreInterprocessMutex, ESPMode::ThreadSafe>(InterprocessMemoryMutexName);
-		if (!InterprocessMemoryMutex.IsValid() || !InterprocessMemoryMutex->Initialize())
+		InterprocessMemoryMutex = MakeShared<FTextureShareCoreInterprocessMutex, ESPMode::ThreadSafe>();
+		if (!InterprocessMemoryMutex.IsValid() || !InterprocessMemoryMutex->InitializeInterprocessMutex(InterprocessMemoryMutexName))
 		{
 			UE_LOG(LogTextureShareCore, Error, TEXT("Can't initialize process shared memory mutex"));
 
diff --git a/Engine/Plugins/VirtualProduction/TextureShare/Source/TextureShareCore/Private/IPC/TextureShareCoreInterprocessMutex.cpp b/Engine/Plugins/VirtualProduction/TextureShare/Source/TextureShareCore/Private/IPC/TextureShareCoreInterprocessMutex.cpp
--- a/Engine/Plugins/VirtualProduction/TextureShare/Source/TextureShareCore/Private/IPC/TextureShareCoreInterprocessMutex.cpp
+++ b/Engine/Plugins/VirtualProduction/TextureShare/Source/TextureShareCore/Private/IPC/TextureShareCoreInterprocessMutex.cpp
@@ -16,100 +16,71 @@ void FTextureShareCoreInterprocessMutex::ReleaseInterprocessMutex()
 {
 	if (PlatformMutex)
 	{
-		TryUnlockMutex();
-
-		FPlatformProcess::FSemaphore* ProcessMutex = static_cast<FPlatformProcess::FSemaphore*>(PlatformMutex);
-		PlatformMutex = nullptr;
+		UnlockMutex();
 
-		if (ProcessMutex)
+		if (PlatformMutex)
 		{
-			LOG_SCOPE_VERBOSITY_OVERRIDE(LogHAL, ELogVerbosity::Error);
-			FWindowsPlatformProcess::DeleteInterprocessSynchObject(ProcessMutex);
+			FWindowsPlatformProcess::DeleteInterprocessSynchObject(PlatformMutex);
+
+			PlatformMutex = nullptr;
 		}
 	}
 }
 
-bool FTextureShareCoreInterprocessMutex::InitializeInterprocessMutex(bool bInGlobalNameSpace)
+bool FTextureShareCoreInterprocessMutex::InitializeInterprocessMutex(bool bInGlobalNameSpace, const FString& InMutexId)
 {
-	if (PlatformMutex == nullptr)
+	if (IsValid())
 	{
-		const bool bCreateUniqueMutex = MutexId.IsEmpty();
-
-		const FString MutexName = bCreateUniqueMutex ? FGuid::NewGuid().ToString(EGuidFormats::DigitsWithHyphensInBraces) : MutexId;
-		const FString FullMutexName = (bInGlobalNameSpace && !bCreateUniqueMutex) ? FString::Printf(TEXT("Global\\%s"), *MutexName) : FString::Printf(TEXT("Local\\%s"), *MutexName);
+		// Already initialized
+		return false;
+	}
 
-		FPlatformProcess::FSemaphore* ProcessMutex = nullptr;
+	const FString FullMutexName = bInGlobalNameSpace
+		? FString::Printf(TEXT("Global\\%s"), *InMutexId)
+		: FString::Printf(TEXT("Local\\%s"), *InMutexId);
 
-		if (!bCreateUniqueMutex)
-		{
-			LOG_SCOPE_VERBOSITY_OVERRIDE(LogHAL, ELogVerbosity::Error);
+	FPlatformProcess::FSemaphore* ProcessMutex = nullptr;
 
-			// Open exist  mutex:
-			ProcessMutex = FPlatformProcess::NewInterprocessSynchObject(*FullMutexName, false);
-			if (ProcessMutex)
-			{
-				PlatformMutex = ProcessMutex;
-				return true;
-			}
-		}
+	// In the global namespace, this mutex may already exist, so we should try to open it first.
+	if (bInGlobalNameSpace)
+	{
+		// Try to open exist  mutex
+		PlatformMutex = FPlatformProcess::NewInterprocessSynchObject(*FullMutexName, false);
+	}
 
-		// Create new:
-		ProcessMutex = FPlatformProcess::NewInterprocessSynchObject(*FullMutexName, true);
-		if (ProcessMutex)
-		{
-			PlatformMutex = ProcessMutex;
-			return true;
-		}
+	// Otherwise try to create a new one.
+	if (!PlatformMutex)
+	{
+		// Create new mutex
+		PlatformMutex = FPlatformProcess::NewInterprocessSynchObject(*FullMutexName, true);
 	}
 
-	return false;
+	return IsValid();
 }
 
 bool FTextureShareCoreInterprocessMutex::LockMutex(const uint32 InMaxMillisecondsToWait)
 {
-	if (PlatformMutex)
+	if (!IsValid())
 	{
-		// 1ms = 10^6 ns
-		const uint64 MaxNanosecondsToWait = InMaxMillisecondsToWait * 1000000ULL;
-
-		FPlatformProcess::FSemaphore* ProcessMutex = static_cast<FPlatformProcess::FSemaphore*>(PlatformMutex);
-		if (ProcessMutex)
-		{
-			// InMaxMillisecondsToWait=0 means a lock attempt without waiting.
-			// Now we have never used an infinite lock to prevent deadlocks.
-			if (!ProcessMutex->TryLock(MaxNanosecondsToWait))
-			{
-				return false;
-			}
-
-			bLocked = true;
-
-			return true;
-		}
+		return false;
 	}
 
-	return false;
-}
+	// Converts miliseconds to the nanoseconds
+	// 1ms = 10^6 ns
+	const uint64 MaxNanosecondsToWait = InMaxMillisecondsToWait * 1000000ULL;
 
-void FTextureShareCoreInterprocessMutex::UnlockMutex()
-{
-	if (PlatformMutex)
-	{
-		FPlatformProcess::FSemaphore* ProcessMutex = static_cast<FPlatformProcess::FSemaphore*>(PlatformMutex);
-		if (ProcessMutex)
-		{
-			// relinquish
-			bLocked = false;
-
-			ProcessMutex->Unlock();
-		}
-	}
+	// InMaxMillisecondsToWait=0 means a lock attempt without waiting.
+	// Now we have never used an infinite lock to prevent deadlocks.
+	return PlatformMutex->TryLock(MaxNanosecondsToWait);
 }
 
-void FTextureShareCoreInterprocessMutex::TryUnlockMutex()
+void FTextureShareCoreInterprocessMutex::UnlockMutex()
 {
-	if (bLocked)
-	{
-		UnlockMutex();
+	if (IsValid())
+	{		
+		// InMaxMillisecondsToWait=0 means a lock attempt without waiting.
+		// Ensure it is locked before we try to unlock it
+		PlatformMutex->TryLock(0);
+		PlatformMutex->Unlock();
 	}
-}
+}
\ No newline at end of file
diff --git a/Engine/Plugins/VirtualProduction/TextureShare/Source/TextureShareCore/Private/IPC/TextureShareCoreInterprocessMutex.h b/Engine/Plugins/VirtualProduction/TextureShare/Source/TextureShareCore/Private/IPC/TextureShareCoreInterprocessMutex.h
--- a/Engine/Plugins/VirtualProduction/TextureShare/Source/TextureShareCore/Private/IPC/TextureShareCoreInterprocessMutex.h
+++ b/Engine/Plugins/VirtualProduction/TextureShare/Source/TextureShareCore/Private/IPC/TextureShareCoreInterprocessMutex.h
@@ -3,47 +3,61 @@
 #pragma once
 
 #include "CoreMinimal.h"
+#include "GenericPlatform/GenericPlatformProcess.h"
 
 /**
- * IPC mutex
+ * The mutex class implementation.
+ * This class can be used in both cases, for processes and threads.
  */
 class FTextureShareCoreInterprocessMutex
 {
 public:
 	FTextureShareCoreInterprocessMutex()
-		: MutexId()
 	{ }
 
-	FTextureShareCoreInterprocessMutex(const FString& InMutexId)
-		: MutexId(InMutexId)
-	{ }
-
-
 	~FTextureShareCoreInterprocessMutex();
 
+	/** Create a interprocess platform mutex object in the global namespace
+	* @param InMutexId - the name of the interprocess mutex
+	*/
+	bool InitializeInterprocessMutex(const FString& InMutexId)
+	{
+		return InitializeInterprocessMutex(/** bInGlobalNameSpace= */ true, InMutexId);
+	}
+
+	/** Create a multithread platform mutex object in the local namespace. */
 	bool Initialize()
 	{
-		return InitializeInterprocessMutex(true);
+		const FString UniqueMutexName = FGuid::NewGuid().ToString(EGuidFormats::DigitsWithHyphensInBraces);
+
+		return InitializeInterprocessMutex(/** bInGlobalNameSpace= */ false, UniqueMutexName);
 	}
 
+	/**
+	* Tries to acquire and exclusive access for a specified amount of nanoseconds (also known as TryWait()).
+	*
+	* @param InMaxMillisecondsToWait - milliseconds to wait for
+	* @return false if was not able to lock within given time
+	*/
 	bool LockMutex(const uint32 InMaxMillisecondsToWait);
-	void UnlockMutex();
 
-	void TryUnlockMutex();
+	/** Relinquishes an exclusive access (also known as Release()) */
+	void UnlockMutex();
 
+	/** Is the platform mutex exist. */
 	bool IsValid() const
 	{
 		return PlatformMutex != nullptr;
 	}
 
 private:
-	bool InitializeInterprocessMutex(bool bGlobalNameSpace);
-	void ReleaseInterprocessMutex();
+	/** Implementation of mutex opening/creation. */
+	bool InitializeInterprocessMutex(bool bGlobalNameSpace, const FString& InMutexId);
 
-public:
-	const FString MutexId;
+	/** Implementation of mutex releasing. */
+	void ReleaseInterprocessMutex();
 
 private:
-	void* PlatformMutex = nullptr;
-	bool bLocked = false;
+	// the reference to the platform mutex
+	FGenericPlatformProcess::FSemaphore* PlatformMutex = nullptr;
 };
diff --git a/Engine/Plugins/VirtualProduction/TextureShare/Source/TextureShareCore/Private/Object/TextureShareCoreObject.cpp b/Engine/Plugins/VirtualProduction/TextureShare/Source/TextureShareCore/Private/Object/TextureShareCoreObject.cpp
--- a/Engine/Plugins/VirtualProduction/TextureShare/Source/TextureShareCore/Private/Object/TextureShareCoreObject.cpp
+++ b/Engine/Plugins/VirtualProduction/TextureShare/Source/TextureShareCore/Private/Object/TextureShareCoreObject.cpp
@@ -65,11 +65,15 @@ void FTextureShareCoreObject::HandleResetSync(FTextureShareCoreInterprocessMemor
 	UE_TS_LOG(LogTextureShareCoreObjectSync, Error, TEXT("%s:ResetSync(%s) %s"), *GetName(), *ToString(LocalObject), *ToString(GetFrameConnections()));
 #endif
 
-	// Unlock multithread thread mutexes
+	// Reset GT+RT threads sync mutexes
 	if (LockThreadMutex(ETextureShareThreadMutex::InternalLock))
 	{
-		ResetThreadMutex(ETextureShareThreadMutex::GameThread);
-		ResetThreadMutex(ETextureShareThreadMutex::RenderingThread);
+		// Special use case for these mutexes: ETextureShareThreadMutex::GameThread and ETextureShareThreadMutex::RenderingThread.
+		// They used to lock threads (GT and RT), not data.
+		// At this point, when we lose the sync, we should immediately reset all logic to the initial (unconnected) state.
+		// This is a very specific use of the mutexes which don't follow the general rule of the lock()+unlock() pairs.
+		UnlockThreadMutex(ETextureShareThreadMutex::GameThread);
+		UnlockThreadMutex(ETextureShareThreadMutex::RenderingThread);
 
 		UnlockThreadMutex(ETextureShareThreadMutex::InternalLock);
 	}
diff --git a/Engine/Plugins/VirtualProduction/TextureShare/Source/TextureShareCore/Private/Object/TextureShareCoreObject.h b/Engine/Plugins/VirtualProduction/TextureShare/Source/TextureShareCore/Private/Object/TextureShareCoreObject.h
--- a/Engine/Plugins/VirtualProduction/TextureShare/Source/TextureShareCore/Private/Object/TextureShareCoreObject.h
+++ b/Engine/Plugins/VirtualProduction/TextureShare/Source/TextureShareCore/Private/Object/TextureShareCoreObject.h
@@ -83,7 +83,6 @@ public:
 
 protected:
 	TSharedPtr<FTextureShareCoreInterprocessMutex, ESPMode::ThreadSafe> GetThreadMutex(const ETextureShareThreadMutex InThreadMutex) const;
-	void ResetThreadMutex(const ETextureShareThreadMutex InThreadMutex);
 
 	void InitializeThreadMutexes();
 	void ReleaseThreadMutexes();
diff --git a/Engine/Plugins/VirtualProduction/TextureShare/Source/TextureShareCore/Private/Object/TextureShareCoreObject_ThreadMutex.cpp b/Engine/Plugins/VirtualProduction/TextureShare/Source/TextureShareCore/Private/Object/TextureShareCoreObject_ThreadMutex.cpp
--- a/Engine/Plugins/VirtualProduction/TextureShare/Source/TextureShareCore/Private/Object/TextureShareCoreObject_ThreadMutex.cpp
+++ b/Engine/Plugins/VirtualProduction/TextureShare/Source/TextureShareCore/Private/Object/TextureShareCoreObject_ThreadMutex.cpp
@@ -88,7 +88,7 @@ bool FTextureShareCoreObject::LockThreadMutex(const ETextureShareThreadMutex InT
 	// mutex deadlock
 	UE_TS_LOG(LogTextureShareCoreObjectSync, Error, TEXT("%s:LockThreadMutex(%s) DEADLOCK"), *GetName(), GetTEXT(InThreadMutex));
 
-	// Release mutex
+	// Unlock stale mutex
 	ThreadMutex->UnlockMutex();
 
 	return false;
@@ -146,29 +146,6 @@ TSharedPtr<FTextureShareCoreInterprocessMutex, ESPMode::ThreadSafe> FTextureShar
 	return nullptr;
 }
 
-void FTextureShareCoreObject::ResetThreadMutex(const ETextureShareThreadMutex InThreadMutex)
-{
-	if (!UE::TextureShareCore::Helpers::IsThreadMutexEnabled(InThreadMutex))
-	{
-		return;
-	}
-
-	TSharedPtr<FTextureShareCoreInterprocessMutex, ESPMode::ThreadSafe> ThreadMutex = GetThreadMutex(InThreadMutex);
-	if (!ThreadMutex.IsValid())
-	{
-		UE_TS_LOG(LogTextureShareCoreObjectSync, Error, TEXT("%s:ResetThreadMutex(%s) Mutex object not exist"), *GetName(), GetTEXT(InThreadMutex));
-
-		return;
-	}
-
-	if (InThreadMutex < ETextureShareThreadMutex::MAX_LOG)
-	{
-		UE_TS_LOG(LogTextureShareCoreObjectSync, Error, TEXT("%s:ResetThreadMutex(%s)"), *GetName(), GetTEXT(InThreadMutex));
-	}
-
-	ThreadMutex->TryUnlockMutex();
-}
-
 void FTextureShareCoreObject::ReleaseThreadMutexes()
 {
 	ThreadMutexes.Empty();
