[TextureShare] [Virtual Production] Added CPU profiler events.

diff --git a/Engine/Plugins/VirtualProduction/TextureShare/Source/TextureShare/Private/Game/WorldSubsystem/TextureShareWorldSubsystem.cpp b/Engine/Plugins/VirtualProduction/TextureShare/Source/TextureShare/Private/Game/WorldSubsystem/TextureShareWorldSubsystem.cpp
--- a/Engine/Plugins/VirtualProduction/TextureShare/Source/TextureShare/Private/Game/WorldSubsystem/TextureShareWorldSubsystem.cpp
+++ b/Engine/Plugins/VirtualProduction/TextureShare/Source/TextureShare/Private/Game/WorldSubsystem/TextureShareWorldSubsystem.cpp
@@ -123,6 +123,8 @@ void UTextureShareWorldSubsystem::Tick(float DeltaTime)
 {
 	if (bWorldPlay)
 	{
+		TRACE_CPUPROFILER_EVENT_SCOPE(TextureShare::WorldSubsystem::Tick);
+
 		if (CVarTextureShareUseWorldSubsystem.GetValueOnGameThread() > 0)
 		{
 			if (TextureShare && TextureShare->IsEnabled())
diff --git a/Engine/Plugins/VirtualProduction/TextureShare/Source/TextureShare/Private/Game/WorldSubsystem/TextureShareWorldSubsystemContext.cpp b/Engine/Plugins/VirtualProduction/TextureShare/Source/TextureShare/Private/Game/WorldSubsystem/TextureShareWorldSubsystemContext.cpp
--- a/Engine/Plugins/VirtualProduction/TextureShare/Source/TextureShare/Private/Game/WorldSubsystem/TextureShareWorldSubsystemContext.cpp
+++ b/Engine/Plugins/VirtualProduction/TextureShare/Source/TextureShare/Private/Game/WorldSubsystem/TextureShareWorldSubsystemContext.cpp
@@ -96,6 +96,8 @@ namespace UE::TextureShareWorldSubsystemContext
 
 void FTextureShareWorldSubsystemContext::Tick_GameThread(ITextureShareObject& TextureShareObject)
 {
+	TRACE_CPUPROFILER_EVENT_SCOPE(TextureShare::WorldSubsystemContext::Tick_GameThread);
+
 	using namespace UE::TextureShareWorldSubsystemContext;
 	if (TextureShareObject.BeginFrameSync())
 	{
diff --git a/Engine/Plugins/VirtualProduction/TextureShare/Source/TextureShare/Private/Object/TextureShareObject.cpp b/Engine/Plugins/VirtualProduction/TextureShare/Source/TextureShare/Private/Object/TextureShareObject.cpp
--- a/Engine/Plugins/VirtualProduction/TextureShare/Source/TextureShare/Private/Object/TextureShareObject.cpp
+++ b/Engine/Plugins/VirtualProduction/TextureShare/Source/TextureShare/Private/Object/TextureShareObject.cpp
@@ -138,6 +138,8 @@ bool FTextureShareObject::IsSessionActive() const
 //////////////////////////////////////////////////////////////////////////////////////////////
 bool FTextureShareObject::BeginFrameSync()
 {
+	TRACE_CPUPROFILER_EVENT_SCOPE_TEXT(*FString::Printf(TEXT("TextureShare::BeginFrameSync(%s)"), *GetName()));
+
 	if (ITextureShareCallbacks::Get().OnTextureSharePreBeginFrameSync().IsBound())
 	{
 		ITextureShareCallbacks::Get().OnTextureSharePreBeginFrameSync().Broadcast(*this);
@@ -178,6 +180,8 @@ bool FTextureShareObject::BeginFrameSync()
 
 bool FTextureShareObject::EndFrameSync()
 {
+	TRACE_CPUPROFILER_EVENT_SCOPE_TEXT(*FString::Printf(TEXT("TextureShare::EndFrameSync(%s)"), *GetName()));
+
 	if (!IsFrameSyncActive())
 	{
 		UE_TS_LOG(LogTextureShareObject, Log, TEXT("%s:EndFrameSync: Canceled"), *GetName());
@@ -223,6 +227,8 @@ bool FTextureShareObject::EndFrameSync()
 
 bool FTextureShareObject::FrameSync(const ETextureShareSyncStep InSyncStep)
 {
+	TRACE_CPUPROFILER_EVENT_SCOPE_TEXT(*FString::Printf(TEXT("TextureShare::FrameSync(%s, %s)"), *GetName(), GetTEXT(InSyncStep)));
+
 	if (IsFrameSyncActive())
 	{
 		UE_TS_LOG(LogTextureShareObject, Log, TEXT("%s:FrameSync(%s)"), *GetName(), GetTEXT(InSyncStep));
diff --git a/Engine/Plugins/VirtualProduction/TextureShare/Source/TextureShare/Private/Object/TextureShareObjectProxy.cpp b/Engine/Plugins/VirtualProduction/TextureShare/Source/TextureShare/Private/Object/TextureShareObjectProxy.cpp
--- a/Engine/Plugins/VirtualProduction/TextureShare/Source/TextureShare/Private/Object/TextureShareObjectProxy.cpp
+++ b/Engine/Plugins/VirtualProduction/TextureShare/Source/TextureShare/Private/Object/TextureShareObjectProxy.cpp
@@ -188,6 +188,8 @@ bool FTextureShareObjectProxy::BeginFrameSync_RenderThread(FRHICommandListImmedi
 {
 	check(IsInRenderingThread());
 
+	TRACE_CPUPROFILER_EVENT_SCOPE_TEXT(*FString::Printf(TEXT("TextureShare::BeginFrameSync_RenderThread(%s)"), *GetName_RenderThread()));
+
 	if (!CoreObject->IsBeginFrameSyncActive_RenderThread())
 	{
 		EnumRemoveFlags(ObjectProxyFlags, ETextureShareObjectProxyFlags::FrameProxySyncActive);
@@ -247,6 +249,8 @@ bool FTextureShareObjectProxy::EndFrameSync_RenderThread(FRHICommandListImmediat
 {
 	check(IsInRenderingThread());
 
+	TRACE_CPUPROFILER_EVENT_SCOPE_TEXT(*FString::Printf(TEXT("TextureShare::EndFrameSync_RenderThread(%s)"), *GetName_RenderThread()));
+
 	if (!EnumHasAnyFlags(ObjectProxyFlags, ETextureShareObjectProxyFlags::FrameProxySyncActive))
 	{
 		return false;
@@ -287,6 +291,8 @@ bool FTextureShareObjectProxy::FrameSync_RenderThread(FRHICommandListImmediate&
 {
 	check(IsInRenderingThread());
 
+	TRACE_CPUPROFILER_EVENT_SCOPE_TEXT(*FString::Printf(TEXT("TextureShare::FrameSync_RenderThread(%s, %s)"), *GetName_RenderThread(), GetTEXT(InSyncStep)));
+
 	if (!ResourcesProxy.IsValid())
 	{
 		UE_TS_LOG(LogTextureShareObjectProxy, Error, TEXT("%s:FrameSync_RenderThread(%s) Failed - Resources proxy invalid. Maybe session isn't started"), *GetName_RenderThread(), GetTEXT(InSyncStep));
@@ -356,7 +362,7 @@ bool FTextureShareObjectProxy::DoFrameSync_RenderThread(FRHICommandListImmediate
 	RHI_BREADCRUMB_EVENT_STAT(RHICmdList, TextureShareObjectProxyFrameSync, "TextureShareObjectProxyFrameSync");
 	SCOPED_GPU_STAT(RHICmdList, TextureShareObjectProxyFrameSync);
 
-	TRACE_CPUPROFILER_EVENT_SCOPE(TextureShare::FrameSync_RenderThread);
+	TRACE_CPUPROFILER_EVENT_SCOPE_TEXT(*FString::Printf(TEXT("TextureShare::DoFrameSync_RenderThread(%s, %s)"), *GetName_RenderThread(), GetTEXT(InSyncStep)));
 
 	// step 1: support mGPU for sender
 	ResourcesProxy->RunCrossGPUTransfer_RenderThread(ECrossGPUTransferType::BeforeSync, RHICmdList, InSyncStep);
@@ -581,7 +587,8 @@ bool FTextureShareObjectProxy::ShareResource_RenderThread(
 
 		RHI_BREADCRUMB_EVENT_STAT(RHICmdList, TextureShareObjectProxyShareResource, "TextureShareObjectProxyShareResource");
 		SCOPED_GPU_STAT(RHICmdList, TextureShareObjectProxyShareResource);
-		TRACE_CPUPROFILER_EVENT_SCOPE(TextureShare::ShareResource_RenderThread);
+
+		TRACE_CPUPROFILER_EVENT_SCOPE_TEXT(*FString::Printf(TEXT("TextureShare::ShareResource_RenderThread(%s)"), *GetName_RenderThread()));
 
 		switch (InResourceRequest.ResourceDesc.OperationType)
 		{
@@ -644,7 +651,8 @@ bool FTextureShareObjectProxy::ShareRenderTargetResource_RenderThread(
 
 		RHI_BREADCRUMB_EVENT_STAT(RHICmdList, TextureShareObjectProxyShareRenderTargetResource, "TextureShareObjectProxyShareRenderTargetResource");
 		SCOPED_GPU_STAT(RHICmdList, TextureShareObjectProxyShareRenderTargetResource);
-		TRACE_CPUPROFILER_EVENT_SCOPE(TextureShare::ShareRenderTargetResource_RenderThread);
+
+		TRACE_CPUPROFILER_EVENT_SCOPE_TEXT(*FString::Printf(TEXT("TextureShare::ShareRenderTargetResource_RenderThread(%s)"), *GetName_RenderThread()));
 
 		switch (InResourceRequest.ResourceDesc.OperationType)
 		{
diff --git a/Engine/Plugins/VirtualProduction/TextureShare/Source/TextureShareCore/Private/Object/TextureShareCoreObject_Connect.cpp b/Engine/Plugins/VirtualProduction/TextureShare/Source/TextureShareCore/Private/Object/TextureShareCoreObject_Connect.cpp
--- a/Engine/Plugins/VirtualProduction/TextureShare/Source/TextureShareCore/Private/Object/TextureShareCoreObject_Connect.cpp
+++ b/Engine/Plugins/VirtualProduction/TextureShare/Source/TextureShareCore/Private/Object/TextureShareCoreObject_Connect.cpp
@@ -21,7 +21,7 @@ using namespace UE::TextureShareCore;
 //////////////////////////////////////////////////////////////////////////////////////////////
 bool FTextureShareCoreObject::TryFrameProcessesConnection(FTextureShareCoreInterprocessMemory& InterprocessMemory, FTextureShareCoreInterprocessObject& LocalObject)
 {
-	TRACE_CPUPROFILER_EVENT_SCOPE(TextureShareCore::Connect);
+	TRACE_CPUPROFILER_EVENT_SCOPE_TEXT(*FString::Printf(TEXT("TextureShareCore::TryFrameProcessesConnection(%s)"), *GetName()));
 
 	// Mark last acces time for current process (this value used by other processes to detect die)
 	LocalObject.Sync.UpdateLastAccessTime();
diff --git a/Engine/Plugins/VirtualProduction/TextureShare/Source/TextureShareCore/Private/Object/TextureShareCoreObject_Sync.cpp b/Engine/Plugins/VirtualProduction/TextureShare/Source/TextureShareCore/Private/Object/TextureShareCoreObject_Sync.cpp
--- a/Engine/Plugins/VirtualProduction/TextureShare/Source/TextureShareCore/Private/Object/TextureShareCoreObject_Sync.cpp
+++ b/Engine/Plugins/VirtualProduction/TextureShare/Source/TextureShareCore/Private/Object/TextureShareCoreObject_Sync.cpp
@@ -54,7 +54,7 @@ bool FTextureShareCoreObject::FindSkippedSyncStepImpl(const ETextureShareSyncSte
 //////////////////////////////////////////////////////////////////////////////////////////////
 bool FTextureShareCoreObject::FrameSync(const ETextureShareSyncStep InSyncStep)
 {
-	TRACE_CPUPROFILER_EVENT_SCOPE(TextureShareCore::FrameSync);
+	TRACE_CPUPROFILER_EVENT_SCOPE_TEXT(*FString::Printf(TEXT("TextureShareCore::FrameSync(%s, %s)"), *GetName(), GetTEXT(InSyncStep)));
 
 	if (!IsFrameSyncActive())
 	{
@@ -95,7 +95,7 @@ bool FTextureShareCoreObject::FrameSync(const ETextureShareSyncStep InSyncStep)
 
 bool FTextureShareCoreObject::FrameSync_RenderThread(const ETextureShareSyncStep InSyncStep)
 {
-	TRACE_CPUPROFILER_EVENT_SCOPE(TextureShareCore::FrameSync_RenderThread);
+	TRACE_CPUPROFILER_EVENT_SCOPE_TEXT(*FString::Printf(TEXT("TextureShareCore::FrameSync_RenderThread(%s, %s)"), *GetName(), GetTEXT(InSyncStep)));
 
 	if (!IsFrameSyncActive_RenderThread())
 	{
@@ -137,6 +137,8 @@ bool FTextureShareCoreObject::FrameSync_RenderThread(const ETextureShareSyncStep
 
 bool FTextureShareCoreObject::BeginFrameSync()
 {
+	TRACE_CPUPROFILER_EVENT_SCOPE_TEXT(*FString::Printf(TEXT("TextureShareCore::BeginFrameSync(%s)"), *GetName()));
+
 	UE_TS_LOG(LogTextureShareCoreProxyObjectSync, Log, TEXT("%s:BeginFrameSync()"), *GetName());
 
 	if (!IsBeginFrameSyncActive())
@@ -163,6 +165,8 @@ bool FTextureShareCoreObject::BeginFrameSync()
 
 bool FTextureShareCoreObject::BeginFrameSync_RenderThread()
 {
+	TRACE_CPUPROFILER_EVENT_SCOPE_TEXT(*FString::Printf(TEXT("TextureShareCore::BeginFrameSync_RenderThread(%s)"), *GetName()));
+
 	UE_TS_LOG(LogTextureShareCoreProxyObjectSync, Log, TEXT("%s:BeginFrameSync_RenderThread()"), *GetName());
 
 	if (!IsBeginFrameSyncActive_RenderThread())
@@ -188,6 +192,8 @@ bool FTextureShareCoreObject::BeginFrameSync_RenderThread()
 
 bool FTextureShareCoreObject::EndFrameSync()
 {
+	TRACE_CPUPROFILER_EVENT_SCOPE_TEXT(*FString::Printf(TEXT("TextureShareCore::EndFrameSync(%s)"), *GetName()));
+
 	if (!IsFrameSyncActive())
 	{
 		UE_TS_LOG(LogTextureShareCoreObjectSync, Error, TEXT("%s:EndFrameSync() failed: no sync for this frame"), *GetName());
@@ -214,6 +220,8 @@ bool FTextureShareCoreObject::EndFrameSync()
 
 bool FTextureShareCoreObject::EndFrameSync_RenderThread()
 {
+	TRACE_CPUPROFILER_EVENT_SCOPE_TEXT(*FString::Printf(TEXT("TextureShareCore::EndFrameSync_RenderThread(%s)"), *GetName()));
+
 	if (!IsFrameSyncActive_RenderThread())
 	{
 		UE_TS_LOG(LogTextureShareCoreProxyObjectSync, Error, TEXT("%s:EndFrameSync_RenderThread() failed: no sync for this frame"), *GetName());
@@ -249,6 +257,8 @@ bool FTextureShareCoreObject::EndFrameSync_RenderThread()
 /////////////////////////////////////////////////////////////////////////////////////////////////////////////////////////////////
 bool FTextureShareCoreObject::DoFrameSync(const ETextureShareSyncStep InSyncStep)
 {
+	TRACE_CPUPROFILER_EVENT_SCOPE_TEXT(*FString::Printf(TEXT("TextureShareCore::DoFrameSync(%s, %s)"), *GetName(), GetTEXT(InSyncStep)));
+
 	if (InSyncStep == ETextureShareSyncStep::FrameFlush)
 	{
 		// Always skip special flush sync pass
@@ -310,6 +320,8 @@ bool FTextureShareCoreObject::DoFrameSync(const ETextureShareSyncStep InSyncStep
 
 bool FTextureShareCoreObject::DoFrameSync_RenderThread(const ETextureShareSyncStep InSyncStep)
 {
+	TRACE_CPUPROFILER_EVENT_SCOPE_TEXT(*FString::Printf(TEXT("TextureShareCore::DoFrameSync_RenderThread(%s, %s)"), *GetName(), GetTEXT(InSyncStep)));
+
 	if (InSyncStep == ETextureShareSyncStep::FrameProxyFlush)
 	{
 		// Always skip special flush sync pass
