VulkanRHI: Remove extensions that don't need to be manually loaded anymore from plugin startup:
- VK_KHR_external_memory is part of Vulkan 1.1 (our minspec)
- VK_KHR_external_memory_capabilities is part of Vulkan 1.1 (our minspec)
- VK_KHR_external_memory_win32/fd are now loaded by default in the engine
- VK_KHR_external_semaphore_win32/fd are now loaded by default in the engine

#rnx

#Synced-CL 51353528

diff --git a/Engine/Plugins/Experimental/AVCodecs/AVCodecsCore/Source/AVCodecsCoreRHI/Private/AVCodecsCoreRHIModule.cpp b/Engine/Plugins/Experimental/AVCodecs/AVCodecsCore/Source/AVCodecsCoreRHI/Private/AVCodecsCoreRHIModule.cpp
--- a/Engine/Plugins/Experimental/AVCodecs/AVCodecsCore/Source/AVCodecsCoreRHI/Private/AVCodecsCoreRHIModule.cpp
+++ b/Engine/Plugins/Experimental/AVCodecs/AVCodecsCore/Source/AVCodecsCoreRHI/Private/AVCodecsCoreRHIModule.cpp
@@ -38,26 +38,6 @@ public:
 	{
 		if (FApp::CanEverRender())
 		{
-#if AVCODECS_USE_VULKAN
-            
-#if PLATFORM_WINDOWS
-			TCHAR const* DynamicRHIModuleName = GetSelectedDynamicRHIModuleName(false);
-#elif PLATFORM_LINUX
-			TCHAR const* DynamicRHIModuleName = TEXT("VulkanRHI");
-#endif
-
-			if (FString("VulkanRHI") == FString(DynamicRHIModuleName))
-			{
-#if PLATFORM_WINDOWS
-				TArray<ANSICHAR const*> const ExtensionsToAdd { VK_KHR_EXTERNAL_MEMORY_EXTENSION_NAME, VK_KHR_EXTERNAL_MEMORY_WIN32_EXTENSION_NAME };
-#elif PLATFORM_LINUX
-				TArray<ANSICHAR const*> const ExtensionsToAdd { VK_KHR_EXTERNAL_MEMORY_EXTENSION_NAME, VK_KHR_EXTERNAL_MEMORY_FD_EXTENSION_NAME };
-#endif
-
-				IVulkanDynamicRHI::AddEnabledDeviceExtensionsAndLayers(ExtensionsToAdd, TArray<ANSICHAR const*>());
-			}
-#endif // AVCODECS_USE_VULKAN
-
 			if (!GEngine || !GEngine->IsInitialized())
 			{
 				PostEngineInitHandle = FCoreDelegates::GetOnPostEngineInit().AddStatic(&FAVCodecsCoreRHI::PostEngineInit);
diff --git a/Engine/Plugins/Media/HardwareEncoders/Source/EncoderNVENC/Private/NVENC_EncoderModule.cpp b/Engine/Plugins/Media/HardwareEncoders/Source/EncoderNVENC/Private/NVENC_EncoderModule.cpp
--- a/Engine/Plugins/Media/HardwareEncoders/Source/EncoderNVENC/Private/NVENC_EncoderModule.cpp
+++ b/Engine/Plugins/Media/HardwareEncoders/Source/EncoderNVENC/Private/NVENC_EncoderModule.cpp
@@ -23,22 +23,6 @@ public:
 
 			if (NVENC.GetIsAvailable())
 			{
-#if PLATFORM_WINDOWS
-                const TCHAR* DynamicRHIModuleName = GetSelectedDynamicRHIModuleName(false);
-#elif PLATFORM_LINUX
-                const TCHAR* DynamicRHIModuleName = TEXT("VulkanRHI");
-#endif
-				
-				if (FString("VulkanRHI") == FString(DynamicRHIModuleName))
-				{					
-#if PLATFORM_WINDOWS
-					const TArray<const ANSICHAR*> ExtentionsToAdd{ VK_KHR_EXTERNAL_MEMORY_EXTENSION_NAME, VK_KHR_EXTERNAL_MEMORY_WIN32_EXTENSION_NAME };
-#elif PLATFORM_LINUX
-					const TArray<const ANSICHAR*> ExtentionsToAdd{ VK_KHR_EXTERNAL_MEMORY_EXTENSION_NAME, VK_KHR_EXTERNAL_MEMORY_FD_EXTENSION_NAME };
-#endif
-					IVulkanDynamicRHI::AddEnabledDeviceExtensionsAndLayers(ExtentionsToAdd, TArray<const ANSICHAR*>());
-				}
-
 				FModuleManager::LoadModuleChecked<FCUDAModule>("CUDA").OnPostCUDAInit.AddLambda([]()
 				{
 					if (IsRHIDeviceNVIDIA())
@@ -47,7 +31,6 @@ public:
 						FVideoEncoderNVENC_H264::Register(FVideoEncoderFactory::Get());
 						PRAGMA_ENABLE_DEPRECATION_WARNINGS
 					}
-					
 				});
 			}
 		}
diff --git a/Engine/Plugins/Media/MediaIOFramework/Source/GPUTextureTransfer/Private/GPUTextureTransferModule.cpp b/Engine/Plugins/Media/MediaIOFramework/Source/GPUTextureTransfer/Private/GPUTextureTransferModule.cpp
--- a/Engine/Plugins/Media/MediaIOFramework/Source/GPUTextureTransfer/Private/GPUTextureTransferModule.cpp
+++ b/Engine/Plugins/Media/MediaIOFramework/Source/GPUTextureTransfer/Private/GPUTextureTransferModule.cpp
@@ -69,20 +69,6 @@ void FGPUTextureTransferModule::StartupModule()
 
 	CVarMediaIOEnableGPUDirect.AsVariable()->OnChangedDelegate().AddRaw(this, &FGPUTextureTransferModule::OnEnableGPUDirectCVarChange);
 
-	if (FApp::CanEverRender())
-	{
-		// Always provide the necessary Vulkan extensions (it will just get ignored if a different RHI is used)
-#if DVP_SUPPORTED_PLATFORM
-		const TArray<const ANSICHAR*> ExtentionsToAdd{ VK_KHR_EXTERNAL_MEMORY_EXTENSION_NAME, VK_KHR_EXTERNAL_MEMORY_WIN32_EXTENSION_NAME, VK_KHR_EXTERNAL_SEMAPHORE_WIN32_EXTENSION_NAME, VK_KHR_SURFACE_EXTENSION_NAME, VK_KHR_WIN32_SURFACE_EXTENSION_NAME };
-#elif PLATFORM_LINUX
-		const TArray<const ANSICHAR*> ExtentionsToAdd{ VK_KHR_EXTERNAL_MEMORY_EXTENSION_NAME, VK_KHR_EXTERNAL_MEMORY_FD_EXTENSION_NAME, VK_KHR_EXTERNAL_SEMAPHORE_FD_EXTENSION_NAME, VK_KHR_SURFACE_EXTENSION_NAME };
-#endif
-
-#if VULKAN_PLATFORM
-		IVulkanDynamicRHI::AddEnabledDeviceExtensionsAndLayers(ExtentionsToAdd, TArray<const ANSICHAR*>());
-#endif
-	}
-
 	// Defer the following since the module is loaded too early for this call to succeed.
 	FCoreDelegates::GetOnPostEngineInit().AddLambda([this]() 
 	{
diff --git a/Engine/Plugins/VirtualProduction/TextureShare/Source/TextureShareCore/Private/Module/TextureShareCoreModule.cpp b/Engine/Plugins/VirtualProduction/TextureShare/Source/TextureShareCore/Private/Module/TextureShareCoreModule.cpp
--- a/Engine/Plugins/VirtualProduction/TextureShare/Source/TextureShareCore/Private/Module/TextureShareCoreModule.cpp
+++ b/Engine/Plugins/VirtualProduction/TextureShare/Source/TextureShareCore/Private/Module/TextureShareCoreModule.cpp
@@ -5,10 +5,6 @@
 
 #include "Core/TextureShareCore.h"
 
-#ifdef TEXTURESHARECORE_INITIALIZEVULKANEXTENSIONS
-#include "IVulkanDynamicRHI.h"
-#endif
-
 //////////////////////////////////////////////////////////////////////////////////////////////
 // FTextureShareCoreModule
 //////////////////////////////////////////////////////////////////////////////////////////////
@@ -26,29 +22,6 @@ ITextureShareCoreAPI& FTextureShareCoreModule::GetTextureShareCoreAPI()
 void FTextureShareCoreModule::StartupModule()
 {
 	UE_LOG(LogTextureShareCore, Log, TEXT("TextureShareCore module startup"));
-
-#ifdef TEXTURESHARECORE_INITIALIZEVULKANEXTENSIONS
-	{
-		// Initialize Vulkan extensions:
-		UE_LOG(LogTextureShareCore, Log, TEXT("TextureShareCore: Initialize Vulkan extensions"));
-
-		const TArray<const ANSICHAR*> DeviceExtensions = {
-			VK_KHR_EXTERNAL_MEMORY_EXTENSION_NAME,
-			VK_KHR_EXTERNAL_MEMORY_WIN32_EXTENSION_NAME
-		};
-		const TArray<const ANSICHAR*> DeviceLayers = {
-		};
-
-		const TArray<const ANSICHAR*> InstanceExtensions = {
-			VK_KHR_EXTERNAL_MEMORY_CAPABILITIES_EXTENSION_NAME
-		};
-		const TArray<const ANSICHAR*> InstanceLayers = {
-		};
-
-		IVulkanDynamicRHI::AddEnabledDeviceExtensionsAndLayers(DeviceExtensions, DeviceLayers);
-		IVulkanDynamicRHI::AddEnabledInstanceExtensionsAndLayers(InstanceExtensions, InstanceLayers);
-	}
-#endif
 }
 
 void FTextureShareCoreModule::ShutdownModule()
